AudioMotorSim plugin pt. 1 - interfaces + types

diff --git a/Engine/Plugins/Runtime/AudioMotorSim/AudioMotorSim.uplugin b/Engine/Plugins/Runtime/AudioMotorSim/AudioMotorSim.uplugin
new file mode 100644
--- /dev/null
+++ b/Engine/Plugins/Runtime/AudioMotorSim/AudioMotorSim.uplugin
@@ -0,0 +1,24 @@
+{
+	"FileVersion": 1,
+	"Version": 1,
+	"VersionName": "1.0",
+	"FriendlyName": "AudioWidgets",
+	"Description": "Compositional method for simulating audio for vehicles.",
+	"Category": "Audio",
+	"CreatedBy": "Epic Games, Inc.",
+	"CreatedByURL": "http://epicgames.com",
+	"DocsURL": "",
+	"MarketplaceURL": "",
+	"SupportURL": "",
+	"CanContainContent": false,
+	"IsBetaVersion": true,
+	"IsExperimentalVersion": true,
+	"Installed": false,
+	"Modules": [
+		{
+			"Name": "AudioMotorSim",
+			"Type": "Runtime",
+			"LoadingPhase": "PreDefault"
+		}
+	]
+}
\ No newline at end of file
diff --git a/Engine/Plugins/Runtime/AudioMotorSim/Source/AudioMotorSim/AudioMotorSim.Build.cs b/Engine/Plugins/Runtime/AudioMotorSim/Source/AudioMotorSim/AudioMotorSim.Build.cs
new file mode 100644
--- /dev/null
+++ b/Engine/Plugins/Runtime/AudioMotorSim/Source/AudioMotorSim/AudioMotorSim.Build.cs
@@ -0,0 +1,20 @@
+// Copyright Epic Games, Inc. All Rights Reserved.
+
+using UnrealBuildTool;
+
+public class AudioMotorSim : ModuleRules
+{
+	public AudioMotorSim(ReadOnlyTargetRules Target) : base(Target)
+	{
+		PCHUsage = ModuleRules.PCHUsageMode.UseExplicitOrSharedPCHs;
+
+		PublicDependencyModuleNames.AddRange(
+		new string[]
+			{
+				"Core",
+				"CoreUObject",
+				"Engine"
+			}
+		);
+	}
+}
diff --git a/Engine/Plugins/Runtime/AudioMotorSim/Source/AudioMotorSim/Private/AudioMotorModelComponent.cpp b/Engine/Plugins/Runtime/AudioMotorSim/Source/AudioMotorSim/Private/AudioMotorModelComponent.cpp
new file mode 100644
--- /dev/null
+++ b/Engine/Plugins/Runtime/AudioMotorSim/Source/AudioMotorSim/Private/AudioMotorModelComponent.cpp
@@ -0,0 +1,108 @@
+// Copyright Epic Games, Inc. All Rights Reserved.
+
+#include "AudioMotorModelComponent.h"
+
+void UAudioMotorModelComponent::Update(FAudioMotorSimInputContext Input)
+{
+	for(FMotorSimEntry& Entry : SimComponents)
+	{
+		if(Entry.Sim)
+		{
+			Entry.Sim->Update(Input, CachedRuntimeContext);
+		}
+	}
+
+	for(TScriptInterface<IAudioMotorSimOutput> Component : AudioComponents)
+	{
+		if(Component)
+		{
+			Component->Update(Input, CachedRuntimeContext);
+		}
+	}
+	
+	CachedInputContext = Input;
+}
+	
+void UAudioMotorModelComponent::Reset()
+{
+	for(FMotorSimEntry& Entry : SimComponents)
+	{
+		if(Entry.Sim)
+		{
+			Entry.Sim->Reset();
+		}
+	}
+}
+
+void UAudioMotorModelComponent::StartOutput()
+{
+	for(TScriptInterface<IAudioMotorSimOutput> Component : AudioComponents)
+	{
+		if(Component)
+		{
+			Component->StartOutput();
+		}
+	}
+}
+	
+void UAudioMotorModelComponent::StopOutput()
+{
+	for(TScriptInterface<IAudioMotorSimOutput> Component : AudioComponents)
+	{
+		if(Component)
+		{
+			Component->StopOutput();
+		}
+	}
+}
+
+void UAudioMotorModelComponent::AddMotorAudioComponent(TScriptInterface<IAudioMotorSimOutput> InComponent)
+{
+	if(InComponent)
+	{
+		AudioComponents.Add(InComponent);
+	}
+}
+	
+void UAudioMotorModelComponent::RemoveMotorAudioComponent(TScriptInterface<IAudioMotorSimOutput> InComponent)
+{
+	if(InComponent)
+	{
+		AudioComponents.Remove(InComponent);
+	}
+}
+	
+void UAudioMotorModelComponent::AddMotorSimComponent(TScriptInterface<IAudioMotorSim> InComponent, const int32 SortOrder)
+{
+	const FMotorSimEntry NewEntry = {InComponent, SortOrder};
+	
+	for(auto SimIt = SimComponents.CreateIterator(); SimIt; ++SimIt)
+	{
+		if(SimIt->SortOrder > SortOrder)
+		{
+			const int32 Index = SimIt.GetIndex();
+			if(SimIt.GetIndex() == 0)
+			{
+				SimComponents.Insert(NewEntry, 0);
+				return;
+			}
+			
+			SimComponents.Insert(NewEntry, Index - 1);
+			return;
+		}
+	}
+	
+	SimComponents.Add(NewEntry);
+}
+	
+void UAudioMotorModelComponent::RemoveMotorSimComponent(TScriptInterface<IAudioMotorSim> InComponent)
+{
+	for(auto SimIt = SimComponents.CreateIterator(); SimIt; ++SimIt)
+	{
+		if(SimIt->Sim == InComponent)
+		{
+			SimIt.RemoveCurrent();
+			return;
+		}
+	}
+}
\ No newline at end of file
diff --git a/Engine/Plugins/Runtime/AudioMotorSim/Source/AudioMotorSim/Private/AudioMotorSimModule.cpp b/Engine/Plugins/Runtime/AudioMotorSim/Source/AudioMotorSim/Private/AudioMotorSimModule.cpp
new file mode 100644
--- /dev/null
+++ b/Engine/Plugins/Runtime/AudioMotorSim/Source/AudioMotorSim/Private/AudioMotorSimModule.cpp
@@ -0,0 +1,14 @@
+// Copyright Epic Games, Inc. All Rights Reserved.
+
+#include "AudioMotorSimModule.h"
+#include "Modules/ModuleInterface.h"
+#include "Modules/ModuleManager.h"
+
+class FAudioMotorSimModule : public IModuleInterface
+{
+public:
+	virtual void StartupModule() override {};
+	virtual void ShutdownModule() override {};
+};
+
+IMPLEMENT_MODULE(FAudioMotorSimModule, AudioMotorSim);
\ No newline at end of file
diff --git a/Engine/Plugins/Runtime/AudioMotorSim/Source/AudioMotorSim/Private/IAudioMotorSim.cpp b/Engine/Plugins/Runtime/AudioMotorSim/Source/AudioMotorSim/Private/IAudioMotorSim.cpp
new file mode 100644
--- /dev/null
+++ b/Engine/Plugins/Runtime/AudioMotorSim/Source/AudioMotorSim/Private/IAudioMotorSim.cpp
@@ -0,0 +1,14 @@
+// Copyright Epic Games, Inc. All Rights Reserved.
+
+#include "IAudioMotorSim.h"
+
+UAudioMotorSim::UAudioMotorSim(const FObjectInitializer& ObjectInitializer)
+	: Super(ObjectInitializer)
+{
+}
+
+UAudioMotorSimComponent::UAudioMotorSimComponent(const FObjectInitializer& ObjectInitializer)
+	: Super(ObjectInitializer)
+{
+	PrimaryComponentTick.bCanEverTick = false;
+}
\ No newline at end of file
diff --git a/Engine/Plugins/Runtime/AudioMotorSim/Source/AudioMotorSim/Private/IAudioMotorSimOutput.cpp b/Engine/Plugins/Runtime/AudioMotorSim/Source/AudioMotorSim/Private/IAudioMotorSimOutput.cpp
new file mode 100644
--- /dev/null
+++ b/Engine/Plugins/Runtime/AudioMotorSim/Source/AudioMotorSim/Private/IAudioMotorSimOutput.cpp
@@ -0,0 +1,8 @@
+// Copyright Epic Games, Inc. All Rights Reserved.
+
+#include "IAudioMotorSimOutput.h"
+
+UAudioMotorSimOutput::UAudioMotorSimOutput(const FObjectInitializer& ObjectInitializer)
+	: Super(ObjectInitializer)
+{
+}
\ No newline at end of file
diff --git a/Engine/Plugins/Runtime/AudioMotorSim/Source/AudioMotorSim/Public/AudioMotorModelComponent.h b/Engine/Plugins/Runtime/AudioMotorSim/Source/AudioMotorSim/Public/AudioMotorModelComponent.h
new file mode 100644
--- /dev/null
+++ b/Engine/Plugins/Runtime/AudioMotorSim/Source/AudioMotorSim/Public/AudioMotorModelComponent.h
@@ -0,0 +1,70 @@
+// Copyright Epic Games, Inc. All Rights Reserved.
+#pragma once
+
+#include "Components/ActorComponent.h"
+#include "AudioMotorSimTypes.h"
+#include "IAudioMotorSim.h"
+#include "IAudioMotorSimOutput.h"
+#include "AudioMotorModelComponent.generated.h"
+
+USTRUCT(BlueprintType)
+struct FMotorSimEntry
+{
+	GENERATED_BODY()
+	
+	TScriptInterface<IAudioMotorSim> Sim;
+	int32 SortOrder = 0;
+};
+	
+UCLASS(ClassGroup = AudioMotorSim, Blueprintable, meta = (BlueprintSpawnableComponent))
+class AUDIOMOTORSIM_API UAudioMotorModelComponent : public UActorComponent
+{
+	GENERATED_BODY()
+public:
+	UPROPERTY(Transient)
+	TArray<FMotorSimEntry> SimComponents;
+
+	UPROPERTY(Transient)
+	TArray<TScriptInterface<IAudioMotorSimOutput>> AudioComponents;
+
+	UFUNCTION(BlueprintCallable, Category = MotorModel)
+	virtual void Update(FAudioMotorSimInputContext Input);
+	
+	UFUNCTION(BlueprintCallable, Category = MotorModel)
+	virtual void Reset();
+
+	UFUNCTION(BlueprintCallable, Category = MotorModel)
+	virtual void StartOutput();
+	
+	UFUNCTION(BlueprintCallable, Category = MotorModel)
+	virtual void StopOutput();
+
+	UFUNCTION(BlueprintCallable, Category = MotorModel)
+	void AddMotorAudioComponent(TScriptInterface<IAudioMotorSimOutput> InComponent);
+	
+	UFUNCTION(BlueprintCallable, Category = MotorModel)
+	void RemoveMotorAudioComponent(TScriptInterface<IAudioMotorSimOutput> InComponent);
+	
+	UFUNCTION(BlueprintCallable, Category = MotorModel)
+	void AddMotorSimComponent(TScriptInterface<IAudioMotorSim> InComponent, const int32 SortOrder = 0);
+	
+	UFUNCTION(BlueprintCallable, Category = MotorModel)
+	void RemoveMotorSimComponent(TScriptInterface<IAudioMotorSim> InComponent);
+	
+	UFUNCTION(BlueprintPure, Category = State)
+	float GetRpm() const { return CachedRuntimeContext.Rpm; }
+
+	UFUNCTION(BlueprintPure, Category = State)
+	int32 GetGear() const { return CachedRuntimeContext.Gear; }
+	
+	UFUNCTION(BlueprintPure, Category = State)
+	FAudioMotorSimRuntimeContext GetRuntimeInfo() const { return CachedRuntimeContext; }
+
+	UFUNCTION(BlueprintPure, Category = State)
+	FAudioMotorSimInputContext GetCachedInputData() const { return CachedInputContext; }
+
+private:
+	FAudioMotorSimRuntimeContext CachedRuntimeContext;
+	
+	FAudioMotorSimInputContext CachedInputContext;
+};
\ No newline at end of file
diff --git a/Engine/Plugins/Runtime/AudioMotorSim/Source/AudioMotorSim/Public/AudioMotorSimModule.h b/Engine/Plugins/Runtime/AudioMotorSim/Source/AudioMotorSim/Public/AudioMotorSimModule.h
new file mode 100644
--- /dev/null
+++ b/Engine/Plugins/Runtime/AudioMotorSim/Source/AudioMotorSim/Public/AudioMotorSimModule.h
@@ -0,0 +1,2 @@
+// Copyright Epic Games, Inc. All Rights Reserved.
+#pragma once
\ No newline at end of file
diff --git a/Engine/Plugins/Runtime/AudioMotorSim/Source/AudioMotorSim/Public/AudioMotorSimTypes.h b/Engine/Plugins/Runtime/AudioMotorSim/Source/AudioMotorSim/Public/AudioMotorSimTypes.h
new file mode 100644
--- /dev/null
+++ b/Engine/Plugins/Runtime/AudioMotorSim/Source/AudioMotorSim/Public/AudioMotorSimTypes.h
@@ -0,0 +1,100 @@
+// Copyright Epic Games, Inc. All Rights Reserved.
+#pragma once
+
+#include "AudioMotorSimTypes.generated.h"
+
+// collection of properties to be filled out by the vehicle in order to update the motor sim
+USTRUCT(BlueprintType)
+struct AUDIOMOTORSIM_API FAudioMotorSimInputContext
+{
+	GENERATED_BODY()
+
+	// Time in Seonds
+	UPROPERTY(EditAnywhere, BlueprintReadWrite, Category = "UpdateContext")
+	float DeltaTime = 0.f;
+
+	// Current speed of the vehicle in any direction
+	UPROPERTY(EditAnywhere, BlueprintReadWrite, Category = "Velocity")
+	float Speed = 0.f;
+	
+	// Current speed of the vehicle relative to its forward direction
+	UPROPERTY(EditAnywhere, BlueprintReadWrite, Category = "Velocity")
+	float ForwardSpeed = 0.f;
+
+	// absolute value of the speed of the vehicle relative to its right direction
+	UPROPERTY(EditAnywhere, BlueprintReadWrite, Category = "Velocity")
+	float SideSpeed = 0.f;
+	
+	// Current speed of the vehicle along the z-axis
+	UPROPERTY(EditAnywhere, BlueprintReadWrite, Category = "Velocity")
+	float UpSpeed = 0.f;
+	
+	// normalized input representing the player wanting to accelerate [-1, 1]
+	UPROPERTY(EditAnywhere, BlueprintReadWrite, Category = "Input")
+	float Throttle = 0.f;
+	
+	// normalized input representing the player wanting to slow down [0, 1]
+	UPROPERTY(EditAnywhere, BlueprintReadWrite, Category = "Input")
+    float Brake = 0.f;
+	
+	// scaling to apply to any behaviors that model surface friction
+	UPROPERTY(EditAnywhere, BlueprintReadWrite, Category = "State")
+	float SurfaceFrictionModifier = 1.f;
+
+	// scaling to apply to any behaviors that model internal motor friction
+	UPROPERTY(EditAnywhere, BlueprintReadWrite, Category = "State")
+	float MotorFrictionModifier = 1.f;
+	
+	// normalized input representing additional thrust beyond normal driving behaviors
+	UPROPERTY(EditAnywhere, BlueprintReadWrite, Category = "State")
+	float Boost = 0.f;
+	
+	// whether a player can drive the vehicle right now
+	UPROPERTY(EditAnywhere, BlueprintReadWrite, Category = "State")
+	bool bDriving = false;
+	
+	// whether the vehicle is firmly on the ground
+	UPROPERTY(EditAnywhere, BlueprintReadWrite, Category = "State")
+	bool bGrounded = false;
+	
+	// whether the motor can freely shift gears
+	UPROPERTY(EditAnywhere, BlueprintReadWrite, Category = "State")
+	bool bCanShift = false;
+
+	// when true, signals that the gears are temporarily disconnected from the motor
+	UPROPERTY(EditAnywhere, BlueprintReadWrite, Category = "State")
+	bool bClutchEngaged = false;
+};
+
+// properties that represent the current state of the motor sim, and persist between updates
+USTRUCT(BlueprintType)
+struct AUDIOMOTORSIM_API FAudioMotorSimRuntimeContext
+{
+	GENERATED_BODY()
+		
+	// true while the motor is performing a gear shift
+	UPROPERTY(EditAnywhere, BlueprintReadWrite, Category = "State")
+	bool bShifting = false;
+		
+	// which gear the motor is in, if it uses gears
+	UPROPERTY(EditAnywhere, BlueprintReadWrite, Category = "State")
+	int32 Gear = 0;
+	
+	// normalized RPM [0-1] of the motor
+	UPROPERTY(EditAnywhere, BlueprintReadWrite, Category = "State")
+	float Rpm = 0.f;
+	
+	// volume to set on the output component
+	UPROPERTY(EditAnywhere, BlueprintReadWrite, Category = "State")
+	float Volume = 1.f;
+	
+	// pitch to set on the output component
+	UPROPERTY(EditAnywhere, BlueprintReadWrite, Category = "State")
+	float Pitch = 1.f;
+};
+
+namespace AudioMotorSim
+{
+	FORCEINLINE float CmsToKmh(const float InSpeed) { return InSpeed * 0.036f; }
+	FORCEINLINE float KmhToCms(const float InSpeed) { return InSpeed / 0.036f; }
+}
diff --git a/Engine/Plugins/Runtime/AudioMotorSim/Source/AudioMotorSim/Public/IAudioMotorSim.h b/Engine/Plugins/Runtime/AudioMotorSim/Source/AudioMotorSim/Public/IAudioMotorSim.h
new file mode 100644
--- /dev/null
+++ b/Engine/Plugins/Runtime/AudioMotorSim/Source/AudioMotorSim/Public/IAudioMotorSim.h
@@ -0,0 +1,34 @@
+// Copyright Epic Games, Inc. All Rights Reserved.
+#pragma once
+
+#include "AudioMotorSimTypes.h"
+#include "Components/ActorComponent.h"
+#include "UObject/Interface.h"
+#include "IAudioMotorSim.generated.h"
+
+UINTERFACE(BlueprintType)
+class UAudioMotorSim : public UInterface
+{
+	GENERATED_UINTERFACE_BODY()
+};
+
+class AUDIOMOTORSIM_API IAudioMotorSim
+{
+	GENERATED_IINTERFACE_BODY()
+
+public:
+	virtual void Update(FAudioMotorSimInputContext& Input, FAudioMotorSimRuntimeContext& RuntimeInfo) = 0;
+
+	virtual void Reset() {};
+};
+
+UCLASS(BlueprintType)
+class AUDIOMOTORSIM_API UAudioMotorSimComponent : public UActorComponent, public IAudioMotorSim
+{
+	GENERATED_BODY()
+
+public:
+	UAudioMotorSimComponent(const FObjectInitializer& ObjectInitializer);
+	
+	virtual void Update(FAudioMotorSimInputContext& Input, FAudioMotorSimRuntimeContext& RuntimeInfo) override {};
+};
\ No newline at end of file
diff --git a/Engine/Plugins/Runtime/AudioMotorSim/Source/AudioMotorSim/Public/IAudioMotorSimOutput.h b/Engine/Plugins/Runtime/AudioMotorSim/Source/AudioMotorSim/Public/IAudioMotorSimOutput.h
new file mode 100644
--- /dev/null
+++ b/Engine/Plugins/Runtime/AudioMotorSim/Source/AudioMotorSim/Public/IAudioMotorSimOutput.h
@@ -0,0 +1,23 @@
+// Copyright Epic Games, Inc. All Rights Reserved.
+#pragma once
+
+#include "AudioMotorSimTypes.h"
+#include "UObject/Interface.h"
+#include "IAudioMotorSimOutput.generated.h"
+
+UINTERFACE()
+class UAudioMotorSimOutput : public UInterface
+{
+	GENERATED_UINTERFACE_BODY()
+};
+
+class AUDIOMOTORSIM_API IAudioMotorSimOutput
+{
+	GENERATED_IINTERFACE_BODY()
+	
+public:
+	virtual void Update(FAudioMotorSimInputContext& Input, FAudioMotorSimRuntimeContext& RuntimeInfo) = 0;
+
+	virtual void StartOutput() = 0;
+	virtual void StopOutput() = 0;
+};
\ No newline at end of file
