AudioMotorSim pt. 2 - Standard Components
 - set of motor sim components that implement a set of standard behaviors using the AudioMotorSim interface

diff --git a/Engine/Plugins/Runtime/AudioMotorSim/AudioMotorSim.uplugin b/Engine/Plugins/Runtime/AudioMotorSim/AudioMotorSim.uplugin
--- a/Engine/Plugins/Runtime/AudioMotorSim/AudioMotorSim.uplugin
+++ b/Engine/Plugins/Runtime/AudioMotorSim/AudioMotorSim.uplugin
@@ -19,6 +19,11 @@
 			"Name": "AudioMotorSim",
 			"Type": "Runtime",
 			"LoadingPhase": "PreDefault"
+		},
+		{
+			"Name": "AudioMotorSimStandardComponents",
+			"Type": "Runtime",
+			"LoadingPhase": "Default"
 		}
 	]
 }
\ No newline at end of file
diff --git a/Engine/Plugins/Runtime/AudioMotorSim/Source/AudioMotorSimStandardComponents/AudioMotorSimStandardComponents.Build.cs b/Engine/Plugins/Runtime/AudioMotorSim/Source/AudioMotorSimStandardComponents/AudioMotorSimStandardComponents.Build.cs
new file mode 100644
--- /dev/null
+++ b/Engine/Plugins/Runtime/AudioMotorSim/Source/AudioMotorSimStandardComponents/AudioMotorSimStandardComponents.Build.cs
@@ -0,0 +1,21 @@
+// Copyright Epic Games, Inc. All Rights Reserved.
+
+using UnrealBuildTool;
+
+public class AudioMotorSimStandardComponents : ModuleRules
+{
+	public AudioMotorSimStandardComponents(ReadOnlyTargetRules Target) : base(Target)
+	{
+		PCHUsage = ModuleRules.PCHUsageMode.UseExplicitOrSharedPCHs;
+
+		PublicDependencyModuleNames.AddRange(
+		new string[]
+			{
+				"AudioMotorSim",
+				"Core",
+				"CoreUObject",
+				"Engine"
+			}
+		);
+	}
+}
diff --git a/Engine/Plugins/Runtime/AudioMotorSim/Source/AudioMotorSimStandardComponents/Private/AudioMotorSimStandardComponentsModule.cpp b/Engine/Plugins/Runtime/AudioMotorSim/Source/AudioMotorSimStandardComponents/Private/AudioMotorSimStandardComponentsModule.cpp
new file mode 100644
--- /dev/null
+++ b/Engine/Plugins/Runtime/AudioMotorSim/Source/AudioMotorSimStandardComponents/Private/AudioMotorSimStandardComponentsModule.cpp
@@ -0,0 +1,11 @@
+// Copyright Epic Games, Inc. All Rights Reserved.
+
+#include "Modules/ModuleInterface.h"
+#include "Modules/ModuleManager.h"
+
+class FAudioMotorSimStandardComponentsModule : public IModuleInterface
+{
+
+};
+
+IMPLEMENT_MODULE(FAudioMotorSimStandardComponentsModule, AudioMotorSimStandardComponents);
\ No newline at end of file
diff --git a/Engine/Plugins/Runtime/AudioMotorSim/Source/AudioMotorSimStandardComponents/Private/BoostMotorSimComponent.cpp b/Engine/Plugins/Runtime/AudioMotorSim/Source/AudioMotorSimStandardComponents/Private/BoostMotorSimComponent.cpp
new file mode 100644
--- /dev/null
+++ b/Engine/Plugins/Runtime/AudioMotorSim/Source/AudioMotorSimStandardComponents/Private/BoostMotorSimComponent.cpp
@@ -0,0 +1,33 @@
+// Copyright Epic Games, Inc. All Rights Reserved.
+
+#include "BoostMotorSimComponent.h"
+#include "Kismet/KismetMathLibrary.h"
+
+void UBoostMotorSimComponent::Update(FAudioMotorSimInputContext& Input, FAudioMotorSimRuntimeContext& RuntimeInfo)
+{
+	if(bModifyPitch)
+	{
+		const float PitchTarget = BoostToPitchCurve.GetRichCurveConst()->Eval(Input.Boost);
+		RuntimeInfo.Pitch = FMath::FInterpTo(RuntimeInfo.Pitch, PitchTarget, Input.DeltaTime, PitchModifierInterpSpeed);
+	}
+
+	if (Input.Boost <= 0.f)
+	{
+		ActiveTime = 0.0f;
+		return;
+	}
+
+	ActiveTime += Input.DeltaTime;
+
+	float InterpScale = ThrottleScale;
+	if(ScaleThrottleWithBoostStrength)
+	{
+		InterpScale *= Input.Boost;
+	}
+	
+	Input.Throttle = UKismetMathLibrary::Ease(InterpScale, 1.0, FMath::Min(ActiveTime / InterpTime, 1.0f), EEasingFunc::EaseIn, InterpExp);
+}
+void UBoostMotorSimComponent::Reset()
+{
+	ActiveTime = 0.f;
+}
\ No newline at end of file
diff --git a/Engine/Plugins/Runtime/AudioMotorSim/Source/AudioMotorSimStandardComponents/Private/MotorPhysicsSimComponent.cpp b/Engine/Plugins/Runtime/AudioMotorSim/Source/AudioMotorSimStandardComponents/Private/MotorPhysicsSimComponent.cpp
new file mode 100644
--- /dev/null
+++ b/Engine/Plugins/Runtime/AudioMotorSim/Source/AudioMotorSimStandardComponents/Private/MotorPhysicsSimComponent.cpp
@@ -0,0 +1,101 @@
+// Copyright Epic Games, Inc. All Rights Reserved.
+
+#include "MotorPhysicsSimComponent.h"
+
+void UMotorPhysicsSimComponent::Update(FAudioMotorSimInputContext& Input, FAudioMotorSimRuntimeContext& RuntimeInfo)
+{
+	if(ensure(Weight > 0.f && EngineGearRatio > 0.f) == false)
+	{
+		return;
+	}
+		
+	const int32 MaxGear = GearRatios.Num() - 1;
+	const float GearRatio = GetGearRatio(RuntimeInfo.Gear, Input.bClutchEngaged);
+	const float Throttle = FMath::Abs(Input.Throttle);
+	
+	const float SpeedKmh = AudioMotorSim::CmsToKmh(Input.Speed);
+
+	int32& Gear = RuntimeInfo.Gear;
+
+	const float ClutchModifier = Input.bClutchEngaged ? ClutchedForceModifier : 1.f;
+	
+	//Forces in Newtons.
+	const float CurrentEngineTorque = EngineTorque * Throttle * GearRatio * ClutchModifier;
+	
+	const float BrakingForce = Input.Brake * BrakingHorsePower;
+	const float WindResistance = WindResistancePerVelocity * FMath::Square(SpeedKmh);
+	const float FrictionResistance = Weight * GroundFriction * Input.SurfaceFrictionModifier;
+	const float EngineResistance = CurrentEngineTorque * FMath::Max(Input.MotorFrictionModifier - Throttle, 0.0f) * EngineFriction;
+	
+	const float NetForce = CurrentEngineTorque - BrakingForce - WindResistance - FrictionResistance - EngineResistance;
+
+	// Convert Rpm to velocity to apply forces
+	float VirtualSpeedKmh =  CalcVelocity(GearRatio, RuntimeInfo.Rpm);
+	VirtualSpeedKmh += ((NetForce / Weight) * Input.DeltaTime);
+	VirtualSpeedKmh = FMath::Max(VirtualSpeedKmh, 0.0f);
+
+	// Now working backwards, what is the engine RPM as a results of this velocity and gear ratio.
+	float Rpm = CalcRpm(GearRatio, VirtualSpeedKmh);
+
+	// Now determine if we need to shift.
+	if (Input.bCanShift && (Rpm > UpShiftMaxRpm))
+	{
+		if (Gear < MaxGear || bUseInfiniteGears)
+		{
+			++Gear;
+			OnGearChangedEvent.Broadcast(Gear);
+		}
+
+		Rpm = CalcRpm(GetGearRatio(Gear, Input.bClutchEngaged), VirtualSpeedKmh);
+	}
+	// If the RPM that we'd be at if we downshift is below DownShiftRPM, then it makes sense to downshift.
+	else if (Gear >= 1 && DownShiftStartRpm > CalcRpm(GearRatios[Gear - 1], VirtualSpeedKmh))
+	{
+		Gear = 0;
+		Rpm = DownShiftStartRpm;
+	}
+
+	RuntimeInfo.Rpm = FMath::FInterpTo(RuntimeInfo.Rpm, Rpm, Input.DeltaTime, RpmInterpSpeed);
+}
+
+float UMotorPhysicsSimComponent::CalcRpm(const float InGearRatio, const float InSpeed) const
+{
+	return InSpeed * InGearRatio / EngineGearRatio;
+}
+
+float UMotorPhysicsSimComponent::CalcVelocity(const float InGearRatio, const float InRpm) const
+{
+	check(InGearRatio > 0.f);
+	return InRpm * EngineGearRatio / InGearRatio;
+}
+
+float UMotorPhysicsSimComponent::GetGearRatio(const int32 InGear, const bool bInClutched) const
+{
+	if(bInClutched)
+	{
+		return ClutchedGearRatio;
+	}
+	
+	int32 Gear = InGear;
+	if(bUseInfiniteGears == false)
+	{
+		Gear = FMath::Clamp(Gear, 0, GearRatios.Num());
+	}
+	
+	return InterpGearRatio(Gear);
+}
+
+float UMotorPhysicsSimComponent::InterpGearRatio(const int32 InGear) const
+{
+	const int32 MaxGear = GearRatios.Num() - 1;
+	if (InGear <= MaxGear)
+	{
+		return GearRatios[InGear];
+	}
+
+	return GearRatios[MaxGear] * FMath::Pow(InfiniteGearRatio, InGear - MaxGear);
+}
+
+void UMotorPhysicsSimComponent::Reset()
+{
+}
\ No newline at end of file
diff --git a/Engine/Plugins/Runtime/AudioMotorSim/Source/AudioMotorSimStandardComponents/Private/ResistanceMotorSimComponent.cpp b/Engine/Plugins/Runtime/AudioMotorSim/Source/AudioMotorSimStandardComponents/Private/ResistanceMotorSimComponent.cpp
new file mode 100644
--- /dev/null
+++ b/Engine/Plugins/Runtime/AudioMotorSim/Source/AudioMotorSimStandardComponents/Private/ResistanceMotorSimComponent.cpp
@@ -0,0 +1,22 @@
+// Copyright Epic Games, Inc. All Rights Reserved.
+
+#include "ResistanceMotorSimComponent.h"
+
+void UResistanceMotorSimComponent::Update(FAudioMotorSimInputContext& Input, FAudioMotorSimRuntimeContext& RuntimeInfo)
+{
+	if(Input.bClutchEngaged || !Input.bDriving || !Input.bGrounded)
+	{
+		return;
+	}
+	
+	if (Input.Speed > MinSpeed)
+	{
+		check(MinSpeed != 0.f);
+
+		const float SideFriction = SideSpeedFrictionCurve.GetRichCurveConst()->Eval(Input.SideSpeed);
+		const float UpSpeedRatio = Input.UpSpeed / Input.Speed;
+		const float ZFriction = UpSpeedMaxFriction * UpSpeedRatio;
+		
+		Input.SurfaceFrictionModifier += ZFriction + SideFriction;
+	}
+}
\ No newline at end of file
diff --git a/Engine/Plugins/Runtime/AudioMotorSim/Source/AudioMotorSimStandardComponents/Private/RevLimiterMotorSimComponent.cpp b/Engine/Plugins/Runtime/AudioMotorSim/Source/AudioMotorSimStandardComponents/Private/RevLimiterMotorSimComponent.cpp
new file mode 100644
--- /dev/null
+++ b/Engine/Plugins/Runtime/AudioMotorSim/Source/AudioMotorSimStandardComponents/Private/RevLimiterMotorSimComponent.cpp
@@ -0,0 +1,55 @@
+// Copyright Epic Games, Inc. All Rights Reserved.
+
+#include "RevLimiterMotorSimComponent.h"
+
+void URevLimiterMotorSimComponent::Update(FAudioMotorSimInputContext& Input, FAudioMotorSimRuntimeContext& RuntimeInfo)
+{
+	if (Input.bDriving && Input.bGrounded && Input.SideSpeed < SideSpeedThreshold)
+	{
+		TimeRemaining = 0.f;
+		TimeInAir = 0.0f;
+		return;
+	}
+
+	Input.bCanShift = false;
+
+	// We've hit the limiter
+	if (RuntimeInfo.Rpm >= LimiterMaxRpm)
+	{
+		TimeRemaining = LimitTime;
+		RuntimeInfo.Rpm = LimiterMaxRpm;
+	}
+
+	if (TimeRemaining > 0.0f)
+	{
+		Input.Throttle = 0.0f;
+		RuntimeInfo.Gear = 0;
+
+		TimeRemaining -= Input.DeltaTime;
+		Input.bClutchEngaged = true;
+	}
+
+	if (Input.bGrounded)
+	{
+		TimeInAir = 0.0f;
+		return;
+	}
+
+	Input.bClutchEngaged = true;
+
+	if (TimeRemaining > 0.0f)
+	{
+		TimeInAir += Input.DeltaTime;
+	}
+
+	if (TimeInAir >= AirMaxThrottleTime)
+	{
+		Input.Throttle = 0.0f;
+	}
+}
+
+void URevLimiterMotorSimComponent::Reset()
+{
+	TimeRemaining = 0.f;
+	TimeInAir = 0.f;
+}
diff --git a/Engine/Plugins/Runtime/AudioMotorSim/Source/AudioMotorSimStandardComponents/Private/ReverseMotorSimComponent.cpp b/Engine/Plugins/Runtime/AudioMotorSim/Source/AudioMotorSimStandardComponents/Private/ReverseMotorSimComponent.cpp
new file mode 100644
--- /dev/null
+++ b/Engine/Plugins/Runtime/AudioMotorSim/Source/AudioMotorSimStandardComponents/Private/ReverseMotorSimComponent.cpp
@@ -0,0 +1,13 @@
+// Copyright Epic Games, Inc. All Rights Reserved.
+
+#include "ReverseMotorSimComponent.h"
+
+void UReverseMotorSimComponent::Update(FAudioMotorSimInputContext& Input, FAudioMotorSimRuntimeContext& RuntimeInfo)
+{
+	if (Input.ForwardSpeed >= 0.f)
+	{
+		return;
+	}
+
+	Input.MotorFrictionModifier *= ReverseEngineResistanceModifier;
+}
diff --git a/Engine/Plugins/Runtime/AudioMotorSim/Source/AudioMotorSimStandardComponents/Private/RpmCurveMotorSimComponent.cpp b/Engine/Plugins/Runtime/AudioMotorSim/Source/AudioMotorSimStandardComponents/Private/RpmCurveMotorSimComponent.cpp
new file mode 100644
--- /dev/null
+++ b/Engine/Plugins/Runtime/AudioMotorSim/Source/AudioMotorSimStandardComponents/Private/RpmCurveMotorSimComponent.cpp
@@ -0,0 +1,58 @@
+// Copyright Epic Games, Inc. All Rights Reserved.
+
+#include "RpmCurveMotorSimComponent.h"
+
+void URpmCurveMotorSimComponent::Update(FAudioMotorSimInputContext& Input, FAudioMotorSimRuntimeContext& RuntimeInfo)
+{
+	if(Gears.Num() == 0)
+	{
+		return;
+	}
+
+	const float SpeedKmh = AudioMotorSim::CmsToKmh(Input.Speed);
+	
+	int32 DesiredGear = 0;
+
+	if(Input.bCanShift)
+	{
+		DesiredGear = GetDesiredGearForSpeed(SpeedKmh);
+	}
+	else
+	{
+		DesiredGear = RuntimeInfo.Gear;
+	}
+
+	const float SpeedCeil = Gears[DesiredGear].SpeedTopThreshold;
+	const float SpeedFloor = DesiredGear > 0 ? Gears[DesiredGear - 1].SpeedTopThreshold : 0.0f;
+	
+	if(DesiredGear != RuntimeInfo.Gear)
+	{
+		RuntimeInfo.Gear = DesiredGear;
+		// todo: delegate SoundGroup->BroadcastEvent(SoundEvents_RL::GearChanged);
+	}
+	
+	const float SpeedRatio = FMath::Clamp(FMath::GetRangePct(SpeedFloor, SpeedCeil, SpeedKmh), 0.0f, 1.0f);
+	const float TargetRpm = Gears[RuntimeInfo.Gear].RpmCurve.GetRichCurveConst()->Eval(SpeedRatio);
+
+	if(InterpSpeed > 0.f)
+	{
+		RuntimeInfo.Rpm = FMath::FInterpTo(RuntimeInfo.Rpm, TargetRpm, Input.DeltaTime, InterpSpeed);
+	}
+	else
+	{
+		RuntimeInfo.Rpm = TargetRpm;
+	}
+}
+
+int32 URpmCurveMotorSimComponent::GetDesiredGearForSpeed(const float Speed) const
+{
+	for(int32 Gear = 0; Gear < Gears.Num(); ++Gear)
+	{
+		if(Speed < Gears[Gear].SpeedTopThreshold)
+		{
+			return Gear;
+		}
+	}
+
+	return 0;
+}
diff --git a/Engine/Plugins/Runtime/AudioMotorSim/Source/AudioMotorSimStandardComponents/Private/ThrottleStateMotorSimComponent.cpp b/Engine/Plugins/Runtime/AudioMotorSim/Source/AudioMotorSimStandardComponents/Private/ThrottleStateMotorSimComponent.cpp
new file mode 100644
--- /dev/null
+++ b/Engine/Plugins/Runtime/AudioMotorSim/Source/AudioMotorSimStandardComponents/Private/ThrottleStateMotorSimComponent.cpp
@@ -0,0 +1,46 @@
+// Copyright Epic Games, Inc. All Rights Reserved.
+
+#include "ThrottleStateMotorSimComponent.h"
+
+void UThrottleStateMotorSimComponent::Update(FAudioMotorSimInputContext& Input, FAudioMotorSimRuntimeContext& RuntimeInfo)
+{
+	const bool bCarIdling = Input.Throttle == 0.f;
+    const bool bCarReversing = Input.Throttle < 0.f;
+    const bool bFullThrottle = Input.Throttle >= 1.f;
+    const bool bCarAccelerating = !bCarIdling && !bCarReversing;
+
+    if (bFullThrottle)
+    {
+    	FullThrottleTime += Input.DeltaTime;
+    }
+    else if (FullThrottleTime > 0.0f)
+    {
+    	if (Input.Throttle <= 0.0f && FullThrottleTime > BlowoffMinThrottleTime)
+    	{
+    		OnEngineBlowoff.Broadcast(FullThrottleTime);
+    		FullThrottleTime = 0.0f;
+    	}
+
+    	FullThrottleTime = FMath::Max(FullThrottleTime - Input.DeltaTime, 0.0f);
+    }
+
+    if (!bPrevCarAccelerating && bCarAccelerating)
+    {
+    	OnThrottleEngaged.Broadcast();
+    }
+    else if (!bPrevCarIdling && bCarIdling)
+    {
+    	OnThrottleReleased.Broadcast();
+    }
+
+    bPrevCarAccelerating = bCarAccelerating;
+    bPrevCarIdling = bCarIdling;
+}
+
+void UThrottleStateMotorSimComponent::Reset()
+{
+	FullThrottleTime = 0.f;
+
+	bPrevCarAccelerating = false;
+	bPrevCarIdling = true;
+}
\ No newline at end of file
diff --git a/Engine/Plugins/Runtime/AudioMotorSim/Source/AudioMotorSimStandardComponents/Private/VelocitySyncMotorSimComponent.cpp b/Engine/Plugins/Runtime/AudioMotorSim/Source/AudioMotorSimStandardComponents/Private/VelocitySyncMotorSimComponent.cpp
new file mode 100644
--- /dev/null
+++ b/Engine/Plugins/Runtime/AudioMotorSim/Source/AudioMotorSimStandardComponents/Private/VelocitySyncMotorSimComponent.cpp
@@ -0,0 +1,38 @@
+// Copyright Epic Games, Inc. All Rights Reserved.
+#include "VelocitySyncMotorSimComponent.h"
+
+void UVelocitySyncMotorSimComponent::Update(FAudioMotorSimInputContext& Input, FAudioMotorSimRuntimeContext& RuntimeInfo)
+{
+	if(Input.bClutchEngaged || !Input.bDriving || !Input.bGrounded)
+	{
+		return;
+	}
+	
+	if (FMath::IsNearlyZero(Input.Throttle)
+	   || (RuntimeInfo.Gear == 0 && FMath::Abs(Input.Throttle) < FirstGearThrottleThreshold))
+	{
+		NoThrottleTimeElapsed += Input.DeltaTime;
+	}
+	else
+	{
+		NoThrottleTimeElapsed = 0.0f;
+	}
+
+	if (NoThrottleTimeElapsed > NoThrottleTime)
+	{
+		InterpTimeLeft = InterpTime;
+	}
+
+	if (Input.Speed < SpeedThreshold || InterpTimeLeft > 0.f)
+	{
+		const float TargetRpm = SpeedToRpmCurve.GetRichCurveConst()->Eval(Input.Speed);
+		RuntimeInfo.Rpm = FMath::FInterpTo(RuntimeInfo.Rpm, TargetRpm, Input.DeltaTime, InterpSpeed);
+		
+		InterpTimeLeft -= Input.DeltaTime;
+	}
+}
+void UVelocitySyncMotorSimComponent::Reset()
+{
+	NoThrottleTimeElapsed = 0.f;
+	InterpTimeLeft = 0.f;
+}
\ No newline at end of file
diff --git a/Engine/Plugins/Runtime/AudioMotorSim/Source/AudioMotorSimStandardComponents/Public/BoostMotorSimComponent.h b/Engine/Plugins/Runtime/AudioMotorSim/Source/AudioMotorSimStandardComponents/Public/BoostMotorSimComponent.h
new file mode 100644
--- /dev/null
+++ b/Engine/Plugins/Runtime/AudioMotorSim/Source/AudioMotorSimStandardComponents/Public/BoostMotorSimComponent.h
@@ -0,0 +1,48 @@
+// Copyright Epic Games, Inc. All Rights Reserved.
+#pragma once
+
+#include "AudioMotorSimTypes.h"
+#include "IAudioMotorSim.h"
+#include "BoostMotorSimComponent.generated.h"
+
+// Uses Boost input to momentarily scale throttle input and pitch output
+UCLASS(ClassGroup = "AudioMotorSim", meta = (BlueprintSpawnableComponent))
+class UBoostMotorSimComponent : public UAudioMotorSimComponent
+{
+	GENERATED_BODY()
+public:
+
+	// Scale the engine torque by this value when boost is active
+	UPROPERTY(EditAnywhere, Category = "Throttle", meta = (ClampMin = "0.1"))
+	float ThrottleScale = 5.f;
+
+	// controls shape of the scaling
+	UPROPERTY(EditAnywhere, Category = "Throttle", meta = (ClampMin = "0.1"))
+	float InterpExp = 4.f;
+
+	// How fast the torque scales when starting to boost
+	UPROPERTY(EditAnywhere, Category = "Throttle", meta = (ClampMin = "0.1"))
+	float InterpTime = 2.f;
+
+	// whether to use the boost strength to scale ThrottleScale, or just check if it is > 0 to apply the throttle scalar
+	UPROPERTY(EditAnywhere, Category = "Throttle", meta = (ClampMin = "0.1"))
+	bool ScaleThrottleWithBoostStrength = false;
+
+	// whether scale the overall pitch by the boost strength
+	UPROPERTY(EditAnywhere, Category = "Pitch")
+	bool bModifyPitch = true;
+	
+	// Speed at which pitch approaches its target value
+	UPROPERTY(EditAnywhere, Category = "Pitch", meta = (EditCondition = "bModifyPitch = true"))
+	float PitchModifierInterpSpeed = 1.f;
+
+	// Curve to derive final pitch value (in playback speed) from boost strength
+	UPROPERTY(EditAnywhere, Category = "Pitch", meta = (EditCondition = "bModifyPitch = true"))
+	FRuntimeFloatCurve BoostToPitchCurve;
+
+	virtual void Update(FAudioMotorSimInputContext& Input, FAudioMotorSimRuntimeContext& RuntimeInfo) override;
+	virtual void Reset() override;
+private:
+
+	float ActiveTime = 0.f;
+};
\ No newline at end of file
diff --git a/Engine/Plugins/Runtime/AudioMotorSim/Source/AudioMotorSimStandardComponents/Public/MotorPhysicsSimComponent.h b/Engine/Plugins/Runtime/AudioMotorSim/Source/AudioMotorSimStandardComponents/Public/MotorPhysicsSimComponent.h
new file mode 100644
--- /dev/null
+++ b/Engine/Plugins/Runtime/AudioMotorSim/Source/AudioMotorSimStandardComponents/Public/MotorPhysicsSimComponent.h
@@ -0,0 +1,82 @@
+// Copyright Epic Games, Inc. All Rights Reserved.
+#pragma once
+
+#include "AudioMotorSimTypes.h"
+#include "IAudioMotorSim.h"
+#include "MotorPhysicsSimComponent.generated.h"
+
+DECLARE_DYNAMIC_MULTICAST_DELEGATE_OneParam(FOnGearChanged, int32, NewGear);
+
+// Uses throttle input to run a physics simulation to drive RPM and shift gears when needed
+UCLASS(ClassGroup = "AudioMotorSim", meta = (BlueprintSpawnableComponent))
+class UMotorPhysicsSimComponent : public UAudioMotorSimComponent
+{
+	GENERATED_BODY()
+public:
+
+	UPROPERTY(EditAnywhere, BlueprintReadWrite, Category = "Force")
+	float Weight = 900.f;
+
+	UPROPERTY(EditAnywhere, BlueprintReadWrite, Category = "Force")
+	float EngineTorque = 2500.f;
+
+	UPROPERTY(EditAnywhere, BlueprintReadWrite, Category = "Force")
+	float BrakingHorsePower = 6000.f;
+
+	UPROPERTY(EditAnywhere, BlueprintReadWrite, Category = "Gears")
+	TArray<float> GearRatios = { 3.5f, 2.f, 1.4f, 1.f, .7f };
+
+	UPROPERTY(EditAnywhere, BlueprintReadWrite, Category = "Gears")
+	float ClutchedGearRatio = 10.f;
+	
+	UPROPERTY(EditAnywhere, BlueprintReadWrite, Category = "Gears")
+	bool bUseInfiniteGears = false;
+	
+	// how much to scale gear ratio per gear past the max gear
+	UPROPERTY(EditAnywhere, BlueprintReadWrite, Category = "Gears")
+	float InfiniteGearRatio = 0.9f;
+	
+	UPROPERTY(EditAnywhere, BlueprintReadWrite, Category = "Gears")
+	float UpShiftMaxRpm = 0.97f;
+
+	UPROPERTY(EditAnywhere, BlueprintReadWrite, Category = "Gears")
+	float DownShiftStartRpm = 0.94f;
+
+	UPROPERTY(EditAnywhere, BlueprintReadWrite, Category = "Resistance")
+	float ClutchedForceModifier = 1.f;
+	
+	UPROPERTY(EditAnywhere, BlueprintReadWrite, Category = "Resistance")
+	float EngineGearRatio = 50.f;
+
+	// How much of the torque is loss due to engine friction
+	UPROPERTY(EditAnywhere, BlueprintReadWrite, Category = "Resistance")
+	float EngineFriction = 0.66f; 
+
+	// Coefficient of Rolling Resistance
+	UPROPERTY(EditAnywhere, BlueprintReadWrite, Category = "Resistance")
+	float GroundFriction = 1.f;
+
+	UPROPERTY(EditAnywhere, BlueprintReadWrite, Category = "Resistance")
+	float WindResistancePerVelocity = 0.015f;
+
+	UPROPERTY(EditAnywhere, BlueprintReadWrite, Category = "PhysicsSim")
+	float ThrottleInterpolationTime = 0.050f;
+
+	UPROPERTY(EditAnywhere, BlueprintReadWrite, Category = "PhysicsSim")
+	float RpmInterpSpeed = 0.0f;
+	
+	UPROPERTY(BlueprintAssignable, Category = "PhysicsSim")
+	FOnGearChanged OnGearChangedEvent;
+
+	virtual void Update(FAudioMotorSimInputContext& Input, FAudioMotorSimRuntimeContext& RuntimeInfo) override;
+	virtual void Reset() override;
+
+private:
+	float CalcRpm(const float InGearRatio, const float InSpeed) const;
+
+	float CalcVelocity(const float InGearRatio, const float InRpm) const;
+
+	float GetGearRatio(const int32 InGear, const bool bInClutched) const;
+
+	float InterpGearRatio(const int32 InGear) const;
+};
\ No newline at end of file
diff --git a/Engine/Plugins/Runtime/AudioMotorSim/Source/AudioMotorSimStandardComponents/Public/ResistanceMotorSimComponent.h b/Engine/Plugins/Runtime/AudioMotorSim/Source/AudioMotorSimStandardComponents/Public/ResistanceMotorSimComponent.h
new file mode 100644
--- /dev/null
+++ b/Engine/Plugins/Runtime/AudioMotorSim/Source/AudioMotorSimStandardComponents/Public/ResistanceMotorSimComponent.h
@@ -0,0 +1,28 @@
+// Copyright Epic Games, Inc. All Rights Reserved.
+#pragma once
+
+#include "AudioMotorSimTypes.h"
+#include "IAudioMotorSim.h"
+#include "ResistanceMotorSimComponent.generated.h"
+
+// Applys additional surface friction based on the angle of the vehicle's velocity
+UCLASS(ClassGroup = "AudioMotorSim", meta = (BlueprintSpawnableComponent))
+class UResistanceMotorSimComponent : public UAudioMotorSimComponent
+{
+	GENERATED_BODY()
+public:
+
+	// How much to increase surface friction when driving straight up. Scales linearly based on driving angle.
+	UPROPERTY(EditAnywhere, Category = "Resistance")
+	float UpSpeedMaxFriction = 1.f;
+
+	// Minimum speed to apply this extra resistance
+	UPROPERTY(EditAnywhere, Category = "Resistance", meta = (ClampMin = "1.0"))
+	float MinSpeed = 100.f;
+
+	// Additional friction to add based on lateral speed
+	UPROPERTY(EditAnywhere, Category = "Resistance")
+	FRuntimeFloatCurve SideSpeedFrictionCurve;
+	
+	virtual void Update(FAudioMotorSimInputContext& Input, FAudioMotorSimRuntimeContext& RuntimeInfo) override;
+};
\ No newline at end of file
diff --git a/Engine/Plugins/Runtime/AudioMotorSim/Source/AudioMotorSimStandardComponents/Public/RevLimiterMotorSimComponent.h b/Engine/Plugins/Runtime/AudioMotorSim/Source/AudioMotorSimStandardComponents/Public/RevLimiterMotorSimComponent.h
new file mode 100644
--- /dev/null
+++ b/Engine/Plugins/Runtime/AudioMotorSim/Source/AudioMotorSimStandardComponents/Public/RevLimiterMotorSimComponent.h
@@ -0,0 +1,38 @@
+// Copyright Epic Games, Inc. All Rights Reserved.
+#pragma once
+
+#include "AudioMotorSimTypes.h"
+#include "IAudioMotorSim.h"
+#include "RevLimiterMotorSimComponent.generated.h"
+
+// Temporarily cuts throttle and reduces RPM when drifting or in the air
+UCLASS(ClassGroup = "AudioMotorSim", meta = (BlueprintSpawnableComponent))
+class URevLimiterMotorSimComponent : public UAudioMotorSimComponent
+{
+	GENERATED_BODY()
+public:
+
+	UPROPERTY(EditAnywhere, BlueprintReadWrite, Category = "RevLimiter")
+	float LimitTime = 0.f;
+	
+	UPROPERTY(EditAnywhere, BlueprintReadWrite, Category = "RevLimiter")
+	float DecelScale = 0.f;
+	
+	UPROPERTY(EditAnywhere, BlueprintReadWrite, Category = "RevLimiter")
+	float AirMaxThrottleTime = 0.f;
+	
+	UPROPERTY(EditAnywhere, BlueprintReadWrite, Category = "RevLimiter")
+	float SideSpeedThreshold = 0.f;
+
+	UPROPERTY(EditAnywhere, BlueprintReadWrite, Category = "RevLimiter")
+	float LimiterMaxRpm = 0.f;
+
+private:
+	// Time remaining where the limiter is forcing throttle down
+	float TimeRemaining = 0.f;
+	float TimeInAir = 0.f;
+
+public:
+	virtual void Update(FAudioMotorSimInputContext& Input, FAudioMotorSimRuntimeContext& RuntimeInfo) override;
+	virtual void Reset() override;
+};
\ No newline at end of file
diff --git a/Engine/Plugins/Runtime/AudioMotorSim/Source/AudioMotorSimStandardComponents/Public/ReverseMotorSimComponent.h b/Engine/Plugins/Runtime/AudioMotorSim/Source/AudioMotorSimStandardComponents/Public/ReverseMotorSimComponent.h
new file mode 100644
--- /dev/null
+++ b/Engine/Plugins/Runtime/AudioMotorSim/Source/AudioMotorSimStandardComponents/Public/ReverseMotorSimComponent.h
@@ -0,0 +1,20 @@
+// Copyright Epic Games, Inc. All Rights Reserved.
+#pragma once
+
+#include "AudioMotorSimTypes.h"
+#include "IAudioMotorSim.h"
+#include "ReverseMotorSimComponent.generated.h"
+
+// Scales engine friction when reversing
+UCLASS(ClassGroup = "AudioMotorSim", meta = (BlueprintSpawnableComponent))
+class UReverseMotorSimComponent : public UAudioMotorSimComponent
+{
+	GENERATED_BODY()
+	
+public:
+	// How much to scale engine friction by when reversing
+	UPROPERTY(EditAnywhere, Category = "Reverse")
+	float ReverseEngineResistanceModifier = 1.f;
+	
+	virtual void Update(FAudioMotorSimInputContext& Input, FAudioMotorSimRuntimeContext& RuntimeInfo) override;
+};
\ No newline at end of file
diff --git a/Engine/Plugins/Runtime/AudioMotorSim/Source/AudioMotorSimStandardComponents/Public/RpmCurveMotorSimComponent.h b/Engine/Plugins/Runtime/AudioMotorSim/Source/AudioMotorSimStandardComponents/Public/RpmCurveMotorSimComponent.h
new file mode 100644
--- /dev/null
+++ b/Engine/Plugins/Runtime/AudioMotorSim/Source/AudioMotorSimStandardComponents/Public/RpmCurveMotorSimComponent.h
@@ -0,0 +1,39 @@
+// Copyright Epic Games, Inc. All Rights Reserved.
+#pragma once
+
+#include "AudioMotorSimTypes.h"
+#include "IAudioMotorSim.h"
+#include "RpmCurveMotorSimComponent.generated.h"
+
+USTRUCT(BlueprintType)
+struct FMotorSimGearCurve
+{
+	GENERATED_BODY()
+
+	// normalized curve, 0-1, of the output RPM from the last gear's top speed to this one's
+	UPROPERTY(EditAnywhere, Category = "Gears")
+	FRuntimeFloatCurve RpmCurve;
+
+	// Speed at which the next gear starts
+	UPROPERTY(EditAnywhere, Category = "Gears")
+	float SpeedTopThreshold = 0.f;
+};
+
+// Derives Gear and RPM directly from speed using hand drawn curves and gear thresholds
+UCLASS(ClassGroup = "AudioMotorSim", meta = (BlueprintSpawnableComponent))
+class URpmCurveMotorSimComponent : public UAudioMotorSimComponent
+{
+	GENERATED_BODY()
+public:
+	
+	virtual void Update(FAudioMotorSimInputContext& Input, FAudioMotorSimRuntimeContext& RuntimeInfo) override;
+	
+	UPROPERTY(EditAnywhere, Category = "Gears")
+	TArray<FMotorSimGearCurve> Gears;
+	
+	UPROPERTY(EditAnywhere, BlueprintReadWrite, Category = "Gears")
+	float InterpSpeed = 0.f;
+
+private:
+	int32 GetDesiredGearForSpeed(const float Speed) const;
+};
\ No newline at end of file
diff --git a/Engine/Plugins/Runtime/AudioMotorSim/Source/AudioMotorSimStandardComponents/Public/ThrottleStateMotorSimComponent.h b/Engine/Plugins/Runtime/AudioMotorSim/Source/AudioMotorSimStandardComponents/Public/ThrottleStateMotorSimComponent.h
new file mode 100644
--- /dev/null
+++ b/Engine/Plugins/Runtime/AudioMotorSim/Source/AudioMotorSimStandardComponents/Public/ThrottleStateMotorSimComponent.h
@@ -0,0 +1,39 @@
+// Copyright Epic Games, Inc. All Rights Reserved.
+#pragma once
+
+#include "AudioMotorSimTypes.h"
+#include "IAudioMotorSim.h"
+#include "ThrottleStateMotorSimComponent.generated.h"
+
+DECLARE_DYNAMIC_MULTICAST_DELEGATE_OneParam(FOnEngineBlowoff, float, BlowoffStrength);
+DECLARE_DYNAMIC_MULTICAST_DELEGATE(FOnThrottleChanged);
+
+// Provides events for when the throttle changes state
+UCLASS(ClassGroup = "AudioMotorSim", meta = (BlueprintSpawnableComponent))
+class UThrottleStateMotorSimComponent : public UAudioMotorSimComponent
+{
+	GENERATED_BODY()
+	
+public:
+	UPROPERTY(BlueprintAssignable, Category = "ThrottleState")
+	FOnThrottleChanged OnThrottleEngaged;
+	
+	UPROPERTY(BlueprintAssignable, Category = "ThrottleState")
+	FOnThrottleChanged OnThrottleReleased;
+
+	// Fires when the throttle is released, keeping track of how long the throttle was held for
+	UPROPERTY(BlueprintAssignable, Category = "ThrottleState")
+	FOnEngineBlowoff OnEngineBlowoff;
+
+	UPROPERTY(EditAnywhere, Category = "ThrottleState")
+	float BlowoffMinThrottleTime = 1.f;
+	
+	virtual void Update(FAudioMotorSimInputContext& Input, FAudioMotorSimRuntimeContext& RuntimeInfo) override;
+	virtual void Reset() override;
+
+private:
+	float FullThrottleTime = 0.f;
+
+	bool bPrevCarAccelerating = false;
+	bool bPrevCarIdling = true;
+};
\ No newline at end of file
diff --git a/Engine/Plugins/Runtime/AudioMotorSim/Source/AudioMotorSimStandardComponents/Public/VelocitySyncMotorSimComponent.h b/Engine/Plugins/Runtime/AudioMotorSim/Source/AudioMotorSimStandardComponents/Public/VelocitySyncMotorSimComponent.h
new file mode 100644
--- /dev/null
+++ b/Engine/Plugins/Runtime/AudioMotorSim/Source/AudioMotorSimStandardComponents/Public/VelocitySyncMotorSimComponent.h
@@ -0,0 +1,40 @@
+// Copyright Epic Games, Inc. All Rights Reserved.
+#pragma once
+
+#include "AudioMotorSimTypes.h"
+#include "IAudioMotorSim.h"
+#include "VelocitySyncMotorSimComponent.generated.h"
+
+// Sets Rpm directly from speed using a curve, if under a speed threshold or if the throttle is released for a period of time
+UCLASS(ClassGroup = "AudioMotorSim", meta = (BlueprintSpawnableComponent))
+class UVelocitySyncMotorSimComponent : public UAudioMotorSimComponent
+{
+	GENERATED_BODY()
+
+public:
+	UPROPERTY(EditAnywhere, BlueprintReadWrite, Category = "VelocitySync")
+	float NoThrottleTime = 0.1f;
+
+	// Speed below which gears will be ignored, and RPM will sync directly to velocity
+	UPROPERTY(EditAnywhere, BlueprintReadWrite, Category = "VelocitySync")
+	float SpeedThreshold = 200.f;
+
+	UPROPERTY(EditAnywhere, Category = "VelocitySync")
+	FRuntimeFloatCurve SpeedToRpmCurve;
+
+	UPROPERTY(EditAnywhere, BlueprintReadWrite, Category = "VelocitySync")
+	float InterpSpeed = 3.f;
+
+	UPROPERTY(EditAnywhere, BlueprintReadWrite, Category = "VelocitySync")
+	float InterpTime = 0.1f;
+
+	UPROPERTY(EditAnywhere, BlueprintReadWrite, Category = "VelocitySync")
+	float FirstGearThrottleThreshold = 0.5f;
+	
+	virtual void Update(FAudioMotorSimInputContext& Input, FAudioMotorSimRuntimeContext& RuntimeInfo) override;
+	virtual void Reset() override;
+
+private:
+	float NoThrottleTimeElapsed = 0.f;
+	float InterpTimeLeft = 0.f;
+};
\ No newline at end of file
