compile fix for non-unity, missing includes

diff --git a/Engine/Plugins/Runtime/AudioMotorSim/Source/AudioMotorSimStandardComponents/Public/BoostMotorSimComponent.h b/Engine/Plugins/Runtime/AudioMotorSim/Source/AudioMotorSimStandardComponents/Public/BoostMotorSimComponent.h
--- a/Engine/Plugins/Runtime/AudioMotorSim/Source/AudioMotorSimStandardComponents/Public/BoostMotorSimComponent.h
+++ b/Engine/Plugins/Runtime/AudioMotorSim/Source/AudioMotorSimStandardComponents/Public/BoostMotorSimComponent.h
@@ -3,6 +3,7 @@
 
 #include "AudioMotorSimTypes.h"
 #include "IAudioMotorSim.h"
+#include "Curves/CurveFloat.h"
 #include "BoostMotorSimComponent.generated.h"
 
 // Uses Boost input to momentarily scale throttle input and pitch output
diff --git a/Engine/Plugins/Runtime/AudioMotorSim/Source/AudioMotorSimStandardComponents/Public/ResistanceMotorSimComponent.h b/Engine/Plugins/Runtime/AudioMotorSim/Source/AudioMotorSimStandardComponents/Public/ResistanceMotorSimComponent.h
--- a/Engine/Plugins/Runtime/AudioMotorSim/Source/AudioMotorSimStandardComponents/Public/ResistanceMotorSimComponent.h
+++ b/Engine/Plugins/Runtime/AudioMotorSim/Source/AudioMotorSimStandardComponents/Public/ResistanceMotorSimComponent.h
@@ -3,6 +3,7 @@
 
 #include "AudioMotorSimTypes.h"
 #include "IAudioMotorSim.h"
+#include "Curves/CurveFloat.h"
 #include "ResistanceMotorSimComponent.generated.h"
 
 // Applys additional surface friction based on the angle of the vehicle's velocity
diff --git a/Engine/Plugins/Runtime/AudioMotorSim/Source/AudioMotorSimStandardComponents/Public/RpmCurveMotorSimComponent.h b/Engine/Plugins/Runtime/AudioMotorSim/Source/AudioMotorSimStandardComponents/Public/RpmCurveMotorSimComponent.h
--- a/Engine/Plugins/Runtime/AudioMotorSim/Source/AudioMotorSimStandardComponents/Public/RpmCurveMotorSimComponent.h
+++ b/Engine/Plugins/Runtime/AudioMotorSim/Source/AudioMotorSimStandardComponents/Public/RpmCurveMotorSimComponent.h
@@ -3,6 +3,7 @@
 
 #include "AudioMotorSimTypes.h"
 #include "IAudioMotorSim.h"
+#include "Curves/CurveFloat.h"
 #include "RpmCurveMotorSimComponent.generated.h"
 
 USTRUCT(BlueprintType)
diff --git a/Engine/Plugins/Runtime/AudioMotorSim/Source/AudioMotorSimStandardComponents/Public/VelocitySyncMotorSimComponent.h b/Engine/Plugins/Runtime/AudioMotorSim/Source/AudioMotorSimStandardComponents/Public/VelocitySyncMotorSimComponent.h
--- a/Engine/Plugins/Runtime/AudioMotorSim/Source/AudioMotorSimStandardComponents/Public/VelocitySyncMotorSimComponent.h
+++ b/Engine/Plugins/Runtime/AudioMotorSim/Source/AudioMotorSimStandardComponents/Public/VelocitySyncMotorSimComponent.h
@@ -3,6 +3,7 @@
 
 #include "AudioMotorSimTypes.h"
 #include "IAudioMotorSim.h"
+#include "Curves/CurveFloat.h"
 #include "VelocitySyncMotorSimComponent.generated.h"
 
 // Sets Rpm directly from speed using a curve, if under a speed threshold or if the throttle is released for a period of time
