[AudioMotorSim] add physics sim option to shift gears down incrementally instead of directly to 0

diff --git a/Engine/Plugins/Runtime/AudioMotorSim/Source/AudioMotorSimStandardComponents/Private/MotorPhysicsSimComponent.cpp b/Engine/Plugins/Runtime/AudioMotorSim/Source/AudioMotorSimStandardComponents/Private/MotorPhysicsSimComponent.cpp
--- a/Engine/Plugins/Runtime/AudioMotorSim/Source/AudioMotorSimStandardComponents/Private/MotorPhysicsSimComponent.cpp
+++ b/Engine/Plugins/Runtime/AudioMotorSim/Source/AudioMotorSimStandardComponents/Private/MotorPhysicsSimComponent.cpp
@@ -49,9 +49,9 @@ void UMotorPhysicsSimComponent::Update(FAudioMotorSimInputContext& Input, FAudio
 		Rpm = CalcRpm(GetGearRatio(Gear, Input.bClutchEngaged), VirtualSpeedKmh);
 	}
 	// If the RPM that we'd be at if we downshift is below DownShiftRPM, then it makes sense to downshift.
-	else if (Gear >= 1 && DownShiftStartRpm > CalcRpm(GearRatios[Gear - 1], VirtualSpeedKmh))
+	else if (Gear >= 1 && DownShiftStartRpm > CalcRpm(GetGearRatio(Gear - 1, Input.bClutchEngaged), VirtualSpeedKmh))
 	{
-		Gear = 0;
+		Gear = bAlwaysDownshiftToZerothGear ? 0 : Gear - 1;
 		Rpm = DownShiftStartRpm;
 	}
 
diff --git a/Engine/Plugins/Runtime/AudioMotorSim/Source/AudioMotorSimStandardComponents/Public/MotorPhysicsSimComponent.h b/Engine/Plugins/Runtime/AudioMotorSim/Source/AudioMotorSimStandardComponents/Public/MotorPhysicsSimComponent.h
--- a/Engine/Plugins/Runtime/AudioMotorSim/Source/AudioMotorSimStandardComponents/Public/MotorPhysicsSimComponent.h
+++ b/Engine/Plugins/Runtime/AudioMotorSim/Source/AudioMotorSimStandardComponents/Public/MotorPhysicsSimComponent.h
@@ -32,6 +32,9 @@ public:
 	UPROPERTY(EditAnywhere, BlueprintReadWrite, Category = "Gears")
 	bool bUseInfiniteGears = false;
 	
+	UPROPERTY(EditAnywhere, BlueprintReadWrite, Category = "Gears")
+	bool bAlwaysDownshiftToZerothGear = false;
+	
 	// how much to scale gear ratio per gear past the max gear
 	UPROPERTY(EditAnywhere, BlueprintReadWrite, Category = "Gears")
 	float InfiniteGearRatio = 0.9f;
