[AudioMotorSim] remove hardcoded setting of Gear to 0 while limiting - clutch behavior achieves the same intended result without overriding gearing from the rest of the sim

diff --git a/Engine/Plugins/Runtime/AudioMotorSim/Source/AudioMotorSimStandardComponents/Private/RevLimiterMotorSimComponent.cpp b/Engine/Plugins/Runtime/AudioMotorSim/Source/AudioMotorSimStandardComponents/Private/RevLimiterMotorSimComponent.cpp
--- a/Engine/Plugins/Runtime/AudioMotorSim/Source/AudioMotorSimStandardComponents/Private/RevLimiterMotorSimComponent.cpp
+++ b/Engine/Plugins/Runtime/AudioMotorSim/Source/AudioMotorSimStandardComponents/Private/RevLimiterMotorSimComponent.cpp
@@ -23,7 +23,6 @@ void URevLimiterMotorSimComponent::Update(FAudioMotorSimInputContext& Input, FAu
 	if (TimeRemaining > 0.0f)
 	{
 		Input.Throttle = 0.0f;
-		RuntimeInfo.Gear = 0;
 
 		TimeRemaining -= Input.DeltaTime;
 		Input.bClutchEngaged = true;
