[AudioMotorSim] Make motor sim components blueprintable

diff --git a/Engine/Plugins/Runtime/AudioMotorSim/Source/AudioMotorSim/Public/IAudioMotorSim.h b/Engine/Plugins/Runtime/AudioMotorSim/Source/AudioMotorSim/Public/IAudioMotorSim.h
--- a/Engine/Plugins/Runtime/AudioMotorSim/Source/AudioMotorSim/Public/IAudioMotorSim.h
+++ b/Engine/Plugins/Runtime/AudioMotorSim/Source/AudioMotorSim/Public/IAudioMotorSim.h
@@ -22,7 +22,7 @@ public:
 	virtual void Reset() {};
 };
 
-UCLASS(BlueprintType)
+UCLASS(Blueprintable, Category = "AudioMotorSim", meta=(BlueprintSpawnableComponent))
 class AUDIOMOTORSIM_API UAudioMotorSimComponent : public UActorComponent, public IAudioMotorSim
 {
 	GENERATED_BODY()
