[AudioMotorSim] fire gearshift events on both up and down shifts from MotorPhysicsSimComponent

diff --git a/Engine/Plugins/Runtime/AudioMotorSim/Source/AudioMotorSimStandardComponents/Private/MotorPhysicsSimComponent.cpp b/Engine/Plugins/Runtime/AudioMotorSim/Source/AudioMotorSimStandardComponents/Private/MotorPhysicsSimComponent.cpp
--- a/Engine/Plugins/Runtime/AudioMotorSim/Source/AudioMotorSimStandardComponents/Private/MotorPhysicsSimComponent.cpp
+++ b/Engine/Plugins/Runtime/AudioMotorSim/Source/AudioMotorSimStandardComponents/Private/MotorPhysicsSimComponent.cpp
@@ -16,7 +16,8 @@ void UMotorPhysicsSimComponent::Update(FAudioMotorSimInputContext& Input, FAudio
 	const float SpeedKmh = AudioMotorSim::CmsToKmh(Input.Speed);
 
 	int32& Gear = RuntimeInfo.Gear;
-
+	const int32 PrevGear = Gear;
+	
 	const float ClutchModifier = Input.bClutchEngaged ? ClutchedForceModifier : 1.f;
 	
 	//Forces in Newtons.
@@ -43,7 +44,6 @@ void UMotorPhysicsSimComponent::Update(FAudioMotorSimInputContext& Input, FAudio
 		if (Gear < MaxGear || bUseInfiniteGears)
 		{
 			++Gear;
-			OnGearChangedEvent.Broadcast(Gear);
 		}
 
 		Rpm = CalcRpm(GetGearRatio(Gear, Input.bClutchEngaged), VirtualSpeedKmh);
@@ -55,6 +55,11 @@ void UMotorPhysicsSimComponent::Update(FAudioMotorSimInputContext& Input, FAudio
 		Rpm = DownShiftStartRpm;
 	}
 
+	if(Gear != PrevGear)
+	{
+		OnGearChangedEvent.Broadcast(Gear);
+	}
+
 	RuntimeInfo.Rpm = FMath::FInterpTo(RuntimeInfo.Rpm, Rpm, Input.DeltaTime, RpmInterpSpeed);
 }
 
