Fix Component sort order in AudioMotorModelComponent

diff --git a/Engine/Plugins/Runtime/AudioMotorSim/Source/AudioMotorSim/Private/AudioMotorModelComponent.cpp b/Engine/Plugins/Runtime/AudioMotorSim/Source/AudioMotorSim/Private/AudioMotorModelComponent.cpp
--- a/Engine/Plugins/Runtime/AudioMotorSim/Source/AudioMotorSim/Private/AudioMotorModelComponent.cpp
+++ b/Engine/Plugins/Runtime/AudioMotorSim/Source/AudioMotorSim/Private/AudioMotorModelComponent.cpp
@@ -89,7 +89,7 @@ void UAudioMotorModelComponent::AddMotorSimComponent(TScriptInterface<IAudioMoto
 				return;
 			}
 			
-			SimComponents.Insert(NewEntry, Index - 1);
+			SimComponents.Insert(NewEntry, Index);
 			return;
 		}
 	}
diff --git a/Engine/Plugins/Runtime/AudioMotorSim/Source/AudioMotorSimStandardComponents/Private/RevLimiterMotorSimComponent.cpp b/Engine/Plugins/Runtime/AudioMotorSim/Source/AudioMotorSimStandardComponents/Private/RevLimiterMotorSimComponent.cpp
--- a/Engine/Plugins/Runtime/AudioMotorSim/Source/AudioMotorSimStandardComponents/Private/RevLimiterMotorSimComponent.cpp
+++ b/Engine/Plugins/Runtime/AudioMotorSim/Source/AudioMotorSimStandardComponents/Private/RevLimiterMotorSimComponent.cpp
@@ -10,9 +10,21 @@ void URevLimiterMotorSimComponent::Update(FAudioMotorSimInputContext& Input, FAu
 	{
 		TimeRemaining = 0.f;
 		TimeInAir = 0.0f;
+
+		if(bActive)
+		{
+			bActive = false;
+			OnRevLimiterStateChanged.Broadcast(bActive);
+		}
 		return;
 	}
 
+	if(bActive == false)
+	{
+		bActive = true;
+		OnRevLimiterStateChanged.Broadcast(bActive);
+	}
+	
 	Input.bCanShift = false;
 
 	// We've hit the limiter
diff --git a/Engine/Plugins/Runtime/AudioMotorSim/Source/AudioMotorSimStandardComponents/Private/RpmCurveMotorSimComponent.cpp b/Engine/Plugins/Runtime/AudioMotorSim/Source/AudioMotorSimStandardComponents/Private/RpmCurveMotorSimComponent.cpp
--- a/Engine/Plugins/Runtime/AudioMotorSim/Source/AudioMotorSimStandardComponents/Private/RpmCurveMotorSimComponent.cpp
+++ b/Engine/Plugins/Runtime/AudioMotorSim/Source/AudioMotorSimStandardComponents/Private/RpmCurveMotorSimComponent.cpp
@@ -11,7 +11,7 @@ void URpmCurveMotorSimComponent::Update(FAudioMotorSimInputContext& Input, FAudi
 		return;
 	}
 
-	const float SpeedKmh = AudioMotorSim::CmsToKmh(Input.Speed);
+	const float SpeedKmh = Input.Speed;
 	
 	int32 DesiredGear = 0;
 
@@ -29,8 +29,16 @@ void URpmCurveMotorSimComponent::Update(FAudioMotorSimInputContext& Input, FAudi
 	
 	if(DesiredGear != RuntimeInfo.Gear)
 	{
+		if(DesiredGear > RuntimeInfo.Gear)
+		{
+			OnUpShift.Broadcast(DesiredGear);
+		}
+		else
+		{
+			OnDownShift.Broadcast(DesiredGear);
+		}
+		
 		RuntimeInfo.Gear = DesiredGear;
-		// todo: delegate SoundGroup->BroadcastEvent(SoundEvents_RL::GearChanged);
 	}
 	
 	const float SpeedRatio = FMath::Clamp(FMath::GetRangePct(SpeedFloor, SpeedCeil, SpeedKmh), 0.0f, 1.0f);
diff --git a/Engine/Plugins/Runtime/AudioMotorSim/Source/AudioMotorSimStandardComponents/Public/RevLimiterMotorSimComponent.h b/Engine/Plugins/Runtime/AudioMotorSim/Source/AudioMotorSimStandardComponents/Public/RevLimiterMotorSimComponent.h
--- a/Engine/Plugins/Runtime/AudioMotorSim/Source/AudioMotorSimStandardComponents/Public/RevLimiterMotorSimComponent.h
+++ b/Engine/Plugins/Runtime/AudioMotorSim/Source/AudioMotorSimStandardComponents/Public/RevLimiterMotorSimComponent.h
@@ -5,6 +5,8 @@
 #include "IAudioMotorSim.h"
 #include "RevLimiterMotorSimComponent.generated.h"
 
+DECLARE_DYNAMIC_MULTICAST_DELEGATE_OneParam(FOnRevLimiterStateChanged, bool, bNewState);
+
 // Temporarily cuts throttle and reduces RPM when drifting or in the air
 UCLASS(ClassGroup = "AudioMotorSim", meta = (BlueprintSpawnableComponent))
 class URevLimiterMotorSimComponent : public UAudioMotorSimComponent
@@ -27,11 +29,16 @@ public:
 	UPROPERTY(EditAnywhere, BlueprintReadWrite, Category = "RevLimiter")
 	float LimiterMaxRpm = 0.f;
 
+	UPROPERTY(BlueprintAssignable)
+	FOnRevLimiterStateChanged OnRevLimiterStateChanged;
+
 private:
 	// Time remaining where the limiter is forcing throttle down
 	float TimeRemaining = 0.f;
 	float TimeInAir = 0.f;
 
+	bool bActive = false;
+
 public:
 	virtual void Update(FAudioMotorSimInputContext& Input, FAudioMotorSimRuntimeContext& RuntimeInfo) override;
 	virtual void Reset() override;
diff --git a/Engine/Plugins/Runtime/AudioMotorSim/Source/AudioMotorSimStandardComponents/Public/RpmCurveMotorSimComponent.h b/Engine/Plugins/Runtime/AudioMotorSim/Source/AudioMotorSimStandardComponents/Public/RpmCurveMotorSimComponent.h
--- a/Engine/Plugins/Runtime/AudioMotorSim/Source/AudioMotorSimStandardComponents/Public/RpmCurveMotorSimComponent.h
+++ b/Engine/Plugins/Runtime/AudioMotorSim/Source/AudioMotorSimStandardComponents/Public/RpmCurveMotorSimComponent.h
@@ -20,6 +20,8 @@ struct FMotorSimGearCurve
 	float SpeedTopThreshold = 0.f;
 };
 
+DECLARE_DYNAMIC_MULTICAST_DELEGATE_OneParam(FOnGearChangedEvent, int32, NewGear);
+
 // Derives Gear and RPM directly from speed using hand drawn curves and gear thresholds
 UCLASS(ClassGroup = "AudioMotorSim", meta = (BlueprintSpawnableComponent))
 class URpmCurveMotorSimComponent : public UAudioMotorSimComponent
@@ -35,6 +37,12 @@ public:
 	UPROPERTY(EditAnywhere, BlueprintReadWrite, Category = "Gears")
 	float InterpSpeed = 0.f;
 
+	UPROPERTY(BlueprintAssignable)
+	FOnGearChangedEvent OnUpShift;
+	
+	UPROPERTY(BlueprintAssignable)
+	FOnGearChangedEvent OnDownShift;
+
 private:
 	int32 GetDesiredGearForSpeed(const float Speed) const;
 };
\ No newline at end of file
