Add OnRevLimiterHit BP event for use in playing a sound when the limiter hits and drops RPM back down

diff --git a/Engine/Plugins/Runtime/AudioMotorSim/Source/AudioMotorSimStandardComponents/Private/RevLimiterMotorSimComponent.cpp b/Engine/Plugins/Runtime/AudioMotorSim/Source/AudioMotorSimStandardComponents/Private/RevLimiterMotorSimComponent.cpp
--- a/Engine/Plugins/Runtime/AudioMotorSim/Source/AudioMotorSimStandardComponents/Private/RevLimiterMotorSimComponent.cpp
+++ b/Engine/Plugins/Runtime/AudioMotorSim/Source/AudioMotorSimStandardComponents/Private/RevLimiterMotorSimComponent.cpp
@@ -32,6 +32,7 @@ void URevLimiterMotorSimComponent::Update(FAudioMotorSimInputContext& Input, FAu
 	{
 		TimeRemaining = LimitTime;
 		RuntimeInfo.Rpm = LimiterMaxRpm;
+		OnRevLimiterHit.Broadcast();
 	}
 
 	if (TimeRemaining > 0.0f)
diff --git a/Engine/Plugins/Runtime/AudioMotorSim/Source/AudioMotorSimStandardComponents/Public/RevLimiterMotorSimComponent.h b/Engine/Plugins/Runtime/AudioMotorSim/Source/AudioMotorSimStandardComponents/Public/RevLimiterMotorSimComponent.h
--- a/Engine/Plugins/Runtime/AudioMotorSim/Source/AudioMotorSimStandardComponents/Public/RevLimiterMotorSimComponent.h
+++ b/Engine/Plugins/Runtime/AudioMotorSim/Source/AudioMotorSimStandardComponents/Public/RevLimiterMotorSimComponent.h
@@ -5,6 +5,7 @@
 #include "IAudioMotorSim.h"
 #include "RevLimiterMotorSimComponent.generated.h"
 
+DECLARE_DYNAMIC_MULTICAST_DELEGATE(FOnRevLimiterHit);
 DECLARE_DYNAMIC_MULTICAST_DELEGATE_OneParam(FOnRevLimiterStateChanged, bool, bNewState);
 
 // Temporarily cuts throttle and reduces RPM when drifting or in the air
@@ -29,6 +30,9 @@ public:
 	UPROPERTY(EditAnywhere, BlueprintReadWrite, Category = "RevLimiter")
 	float LimiterMaxRpm = 0.f;
 
+	UPROPERTY(BlueprintAssignable)
+	FOnRevLimiterHit OnRevLimiterHit;
+
 	UPROPERTY(BlueprintAssignable)
 	FOnRevLimiterStateChanged OnRevLimiterStateChanged;
 
