Add blueprint events to base AudioMotorSim component

diff --git a/Engine/Plugins/Runtime/AudioMotorSim/Source/AudioMotorSim/Private/IAudioMotorSim.cpp b/Engine/Plugins/Runtime/AudioMotorSim/Source/AudioMotorSim/Private/IAudioMotorSim.cpp
--- a/Engine/Plugins/Runtime/AudioMotorSim/Source/AudioMotorSim/Private/IAudioMotorSim.cpp
+++ b/Engine/Plugins/Runtime/AudioMotorSim/Source/AudioMotorSim/Private/IAudioMotorSim.cpp
@@ -14,3 +14,13 @@ UAudioMotorSimComponent::UAudioMotorSimComponent(const FObjectInitializer& Objec
 {
 	PrimaryComponentTick.bCanEverTick = false;
 }
+
+void UAudioMotorSimComponent::Update(FAudioMotorSimInputContext& Input, FAudioMotorSimRuntimeContext& RuntimeInfo)
+{
+	BP_Update(Input, RuntimeInfo);
+}
+
+void UAudioMotorSimComponent::Reset()
+{
+	BP_Reset();
+}
diff --git a/Engine/Plugins/Runtime/AudioMotorSim/Source/AudioMotorSim/Public/IAudioMotorSim.h b/Engine/Plugins/Runtime/AudioMotorSim/Source/AudioMotorSim/Public/IAudioMotorSim.h
--- a/Engine/Plugins/Runtime/AudioMotorSim/Source/AudioMotorSim/Public/IAudioMotorSim.h
+++ b/Engine/Plugins/Runtime/AudioMotorSim/Source/AudioMotorSim/Public/IAudioMotorSim.h
@@ -23,7 +23,7 @@ public:
 	virtual void Reset() {};
 };
 
-UCLASS(Blueprintable, Category = "AudioMotorSim", meta=(BlueprintSpawnableComponent))
+UCLASS(Abstract, Blueprintable, Category = "AudioMotorSim", meta=(BlueprintSpawnableComponent))
 class AUDIOMOTORSIM_API UAudioMotorSimComponent : public UActorComponent, public IAudioMotorSim
 {
 	GENERATED_BODY()
@@ -31,9 +31,19 @@ class AUDIOMOTORSIM_API UAudioMotorSimComponent : public UActorComponent, public
 public:
 	UAudioMotorSimComponent(const FObjectInitializer& ObjectInitializer);
 	
-	virtual void Update(FAudioMotorSimInputContext& Input, FAudioMotorSimRuntimeContext& RuntimeInfo) override {};
+	virtual void Update(FAudioMotorSimInputContext& Input, FAudioMotorSimRuntimeContext& RuntimeInfo) override;
+	
+	virtual void Reset() override;
+
+	// Called every tick that this component is being updated. Use "Set Members in Struct" to update values for future components in the chain
+	UFUNCTION(BlueprintImplementableEvent, Category = "AudioMotorSim", DisplayName = "Update")
+	void BP_Update(UPARAM(ref) FAudioMotorSimInputContext& Input, UPARAM(ref) FAudioMotorSimRuntimeContext& RuntimeInfo);
+	
+	// Use to reset any state that might be desired. Will be called automatically if the entire MotorSim is Reset, or call it manually
+	UFUNCTION(BlueprintImplementableEvent, BlueprintCallable, Category="AudioMotorSim", DisplayName = "Reset")
+	void BP_Reset();
 };
 
 #if UE_ENABLE_INCLUDE_ORDER_DEPRECATED_IN_5_2
 #include "AudioMotorSimTypes.h"
-#endif
+#endif
\ No newline at end of file
