add unused return type to AudioMotorSim blueprint update function to force it to be generated as a function instead of an Event, which allows changes to the by-ref arguments to be maintained in the native update call.

#rnx

diff --git a/Engine/Plugins/Runtime/AudioMotorSim/Source/AudioMotorSim/Public/IAudioMotorSim.h b/Engine/Plugins/Runtime/AudioMotorSim/Source/AudioMotorSim/Public/IAudioMotorSim.h
--- a/Engine/Plugins/Runtime/AudioMotorSim/Source/AudioMotorSim/Public/IAudioMotorSim.h
+++ b/Engine/Plugins/Runtime/AudioMotorSim/Source/AudioMotorSim/Public/IAudioMotorSim.h
@@ -35,9 +35,13 @@ public:
 	
 	virtual void Reset() override;
 
-	// Called every tick that this component is being updated. Use "Set Members in Struct" to update values for future components in the chain
+	/* Called every tick that this component is being updated. Use "Set Members in Struct" to update values for future components in the chain. The return value does nothing.
+	* @param Input			Holds values which are not saved between update frames which represent input to the simulation
+	* @param RuntimeInfo	Holds values which are saved between update frames to represent the output or state of the simulation
+	* @return				Vestigial, does nothing.
+	*/
 	UFUNCTION(BlueprintImplementableEvent, Category = "AudioMotorSim", DisplayName = "Update")
-	void BP_Update(UPARAM(ref) FAudioMotorSimInputContext& Input, UPARAM(ref) FAudioMotorSimRuntimeContext& RuntimeInfo);
+	bool BP_Update(UPARAM(ref) FAudioMotorSimInputContext& Input, UPARAM(ref) FAudioMotorSimRuntimeContext& RuntimeInfo);
 	
 	// Use to reset any state that might be desired. Will be called automatically if the entire MotorSim is Reset, or call it manually
 	UFUNCTION(BlueprintImplementableEvent, BlueprintCallable, Category="AudioMotorSim", DisplayName = "Reset")
