return highest gear instead of 0 in RpmCurveMotorSimComponent if speed is too high for any existing gears.

#rnx

diff --git a/Engine/Plugins/Runtime/AudioMotorSim/Source/AudioMotorSimStandardComponents/Private/RpmCurveMotorSimComponent.cpp b/Engine/Plugins/Runtime/AudioMotorSim/Source/AudioMotorSimStandardComponents/Private/RpmCurveMotorSimComponent.cpp
--- a/Engine/Plugins/Runtime/AudioMotorSim/Source/AudioMotorSimStandardComponents/Private/RpmCurveMotorSimComponent.cpp
+++ b/Engine/Plugins/Runtime/AudioMotorSim/Source/AudioMotorSimStandardComponents/Private/RpmCurveMotorSimComponent.cpp
@@ -7,7 +7,7 @@
 
 void URpmCurveMotorSimComponent::Update(FAudioMotorSimInputContext& Input, FAudioMotorSimRuntimeContext& RuntimeInfo)
 {
-	if(Gears.Num() == 0)
+	if (Gears.Num() == 0)
 	{
 		return;
 	}
@@ -16,7 +16,7 @@ void URpmCurveMotorSimComponent::Update(FAudioMotorSimInputContext& Input, FAudi
 	
 	int32 DesiredGear = 0;
 
-	if(Input.bCanShift)
+	if (Input.bCanShift)
 	{
 		DesiredGear = GetDesiredGearForSpeed(SpeedKmh);
 	}
@@ -28,9 +28,9 @@ void URpmCurveMotorSimComponent::Update(FAudioMotorSimInputContext& Input, FAudi
 	const float SpeedCeil = Gears[DesiredGear].SpeedTopThreshold;
 	const float SpeedFloor = DesiredGear > 0 ? Gears[DesiredGear - 1].SpeedTopThreshold : 0.0f;
 	
-	if(DesiredGear != RuntimeInfo.Gear)
+	if (DesiredGear != RuntimeInfo.Gear)
 	{
-		if(DesiredGear > RuntimeInfo.Gear)
+		if (DesiredGear > RuntimeInfo.Gear)
 		{
 			OnUpShift.Broadcast(DesiredGear);
 		}
@@ -45,7 +45,7 @@ void URpmCurveMotorSimComponent::Update(FAudioMotorSimInputContext& Input, FAudi
 	const float SpeedRatio = FMath::Clamp(FMath::GetRangePct(SpeedFloor, SpeedCeil, SpeedKmh), 0.0f, 1.0f);
 	const float TargetRpm = Gears[RuntimeInfo.Gear].RpmCurve.GetRichCurveConst()->Eval(SpeedRatio);
 
-	if(InterpSpeed > 0.f)
+	if (InterpSpeed > 0.f)
 	{
 		RuntimeInfo.Rpm = FMath::FInterpTo(RuntimeInfo.Rpm, TargetRpm, Input.DeltaTime, InterpSpeed);
 	}
@@ -57,14 +57,19 @@ void URpmCurveMotorSimComponent::Update(FAudioMotorSimInputContext& Input, FAudi
 
 int32 URpmCurveMotorSimComponent::GetDesiredGearForSpeed(const float Speed) const
 {
-	for(int32 Gear = 0; Gear < Gears.Num(); ++Gear)
+	if (Gears.IsEmpty())
 	{
-		if(Speed < Gears[Gear].SpeedTopThreshold)
+		return 0;
+	}
+	
+	for (int32 Gear = 0; Gear < Gears.Num(); ++Gear)
+	{
+		if (Speed < Gears[Gear].SpeedTopThreshold)
 		{
 			return Gear;
 		}
 	}
 
-	return 0;
+	return Gears.Num() - 1;
 }
 
