Cache AudioMotorSim data in editor so state can be queried in editor widgets

diff --git a/Engine/Plugins/Runtime/AudioMotorSim/Source/AudioMotorSim/Private/IAudioMotorSim.cpp b/Engine/Plugins/Runtime/AudioMotorSim/Source/AudioMotorSim/Private/IAudioMotorSim.cpp
--- a/Engine/Plugins/Runtime/AudioMotorSim/Source/AudioMotorSim/Private/IAudioMotorSim.cpp
+++ b/Engine/Plugins/Runtime/AudioMotorSim/Source/AudioMotorSim/Private/IAudioMotorSim.cpp
@@ -18,9 +18,27 @@ UAudioMotorSimComponent::UAudioMotorSimComponent(const FObjectInitializer& Objec
 void UAudioMotorSimComponent::Update(FAudioMotorSimInputContext& Input, FAudioMotorSimRuntimeContext& RuntimeInfo)
 {
 	BP_Update(Input, RuntimeInfo);
+
+#if WITH_EDITORONLY_DATA
+	CachedInput = Input;
+	CachedRuntimeInfo = RuntimeInfo;
+#endif
 }
 
 void UAudioMotorSimComponent::Reset()
 {
 	BP_Reset();
+
+#if WITH_EDITORONLY_DATA
+	CachedInput = FAudioMotorSimInputContext();
+	CachedRuntimeInfo = FAudioMotorSimRuntimeContext();
+#endif
+}
+
+#if WITH_EDITORONLY_DATA
+void UAudioMotorSimComponent::GetCachedData(FAudioMotorSimInputContext& OutInput, FAudioMotorSimRuntimeContext& OutRuntimeInfo)
+{
+	OutInput = CachedInput;
+	OutRuntimeInfo = CachedRuntimeInfo;
 }
+#endif
diff --git a/Engine/Plugins/Runtime/AudioMotorSim/Source/AudioMotorSim/Public/AudioMotorModelComponent.h b/Engine/Plugins/Runtime/AudioMotorSim/Source/AudioMotorSim/Public/AudioMotorModelComponent.h
--- a/Engine/Plugins/Runtime/AudioMotorSim/Source/AudioMotorSim/Public/AudioMotorModelComponent.h
+++ b/Engine/Plugins/Runtime/AudioMotorSim/Source/AudioMotorSim/Public/AudioMotorModelComponent.h
@@ -12,8 +12,11 @@ USTRUCT(BlueprintType)
 struct FMotorSimEntry
 {
 	GENERATED_BODY()
-	
+
+	UPROPERTY(BlueprintReadOnly, Category = "MotorSim")
 	TScriptInterface<IAudioMotorSim> Sim;
+	
+	UPROPERTY(BlueprintReadOnly, Category = "MotorSim")
 	int32 SortOrder = 0;
 };
 	
@@ -22,7 +25,7 @@ class AUDIOMOTORSIM_API UAudioMotorModelComponent : public UActorComponent
 {
 	GENERATED_BODY()
 public:
-	UPROPERTY(Transient)
+	UPROPERTY(BlueprintReadOnly, Transient, Category = "Model")
 	TArray<FMotorSimEntry> SimComponents;
 
 	UPROPERTY(Transient)
diff --git a/Engine/Plugins/Runtime/AudioMotorSim/Source/AudioMotorSim/Public/IAudioMotorSim.h b/Engine/Plugins/Runtime/AudioMotorSim/Source/AudioMotorSim/Public/IAudioMotorSim.h
--- a/Engine/Plugins/Runtime/AudioMotorSim/Source/AudioMotorSim/Public/IAudioMotorSim.h
+++ b/Engine/Plugins/Runtime/AudioMotorSim/Source/AudioMotorSim/Public/IAudioMotorSim.h
@@ -1,6 +1,7 @@
 // Copyright Epic Games, Inc. All Rights Reserved.
 #pragma once
 
+#include "AudioMotorSimTypes.h"
 #include "Components/ActorComponent.h"
 #include "IAudioMotorSim.generated.h"
 
@@ -20,7 +21,7 @@ class AUDIOMOTORSIM_API IAudioMotorSim
 public:
 	virtual void Update(FAudioMotorSimInputContext& Input, FAudioMotorSimRuntimeContext& RuntimeInfo) = 0;
 
-	virtual void Reset() {};
+	virtual void Reset() {}
 };
 
 UCLASS(Abstract, Blueprintable, Category = "AudioMotorSim", meta=(BlueprintSpawnableComponent))
@@ -46,8 +47,16 @@ public:
 	// Use to reset any state that might be desired. Will be called automatically if the entire MotorSim is Reset, or call it manually
 	UFUNCTION(BlueprintImplementableEvent, BlueprintCallable, Category="AudioMotorSim", DisplayName = "Reset")
 	void BP_Reset();
-};
 
-#if UE_ENABLE_INCLUDE_ORDER_DEPRECATED_IN_5_2
-#include "AudioMotorSimTypes.h"
-#endif
\ No newline at end of file
+#if WITH_EDITORONLY_DATA
+	// Input data after running this component
+	UPROPERTY(BlueprintReadOnly, Category = "DebugInfo")
+	FAudioMotorSimInputContext CachedInput;
+
+	// runtime info after running this component
+	UPROPERTY(BlueprintReadOnly, Category = "DebugInfo")
+	FAudioMotorSimRuntimeContext CachedRuntimeInfo;
+
+	virtual void GetCachedData(FAudioMotorSimInputContext& OutInput, FAudioMotorSimRuntimeContext& OutRuntimeInfo);
+#endif
+};
\ No newline at end of file
diff --git a/Engine/Plugins/Runtime/AudioMotorSim/Source/AudioMotorSimStandardComponents/Private/BoostMotorSimComponent.cpp b/Engine/Plugins/Runtime/AudioMotorSim/Source/AudioMotorSimStandardComponents/Private/BoostMotorSimComponent.cpp
--- a/Engine/Plugins/Runtime/AudioMotorSim/Source/AudioMotorSimStandardComponents/Private/BoostMotorSimComponent.cpp
+++ b/Engine/Plugins/Runtime/AudioMotorSim/Source/AudioMotorSimStandardComponents/Private/BoostMotorSimComponent.cpp
@@ -17,6 +17,7 @@ void UBoostMotorSimComponent::Update(FAudioMotorSimInputContext& Input, FAudioMo
 	if (Input.Boost <= 0.f)
 	{
 		ActiveTime = 0.0f;
+		Super::Update(Input, RuntimeInfo);
 		return;
 	}
 
@@ -29,8 +30,11 @@ void UBoostMotorSimComponent::Update(FAudioMotorSimInputContext& Input, FAudioMo
 	}
 	
 	Input.Throttle = UKismetMathLibrary::Ease(InterpScale, 1.0, FMath::Min(ActiveTime / InterpTime, 1.0f), EEasingFunc::EaseIn, InterpExp);
+
+	Super::Update(Input, RuntimeInfo);
 }
 void UBoostMotorSimComponent::Reset()
 {
+	Super::Reset();
 	ActiveTime = 0.f;
 }
diff --git a/Engine/Plugins/Runtime/AudioMotorSim/Source/AudioMotorSimStandardComponents/Private/MotorPhysicsSimComponent.cpp b/Engine/Plugins/Runtime/AudioMotorSim/Source/AudioMotorSimStandardComponents/Private/MotorPhysicsSimComponent.cpp
--- a/Engine/Plugins/Runtime/AudioMotorSim/Source/AudioMotorSimStandardComponents/Private/MotorPhysicsSimComponent.cpp
+++ b/Engine/Plugins/Runtime/AudioMotorSim/Source/AudioMotorSimStandardComponents/Private/MotorPhysicsSimComponent.cpp
@@ -64,6 +64,8 @@ void UMotorPhysicsSimComponent::Update(FAudioMotorSimInputContext& Input, FAudio
 	}
 
 	RuntimeInfo.Rpm = FMath::FInterpTo(RuntimeInfo.Rpm, Rpm, Input.DeltaTime, RpmInterpSpeed);
+
+	Super::Update(Input, RuntimeInfo);
 }
 
 float UMotorPhysicsSimComponent::CalcRpm(const float InGearRatio, const float InSpeed) const
@@ -103,7 +105,3 @@ float UMotorPhysicsSimComponent::InterpGearRatio(const int32 InGear) const
 
 	return GearRatios[MaxGear] * FMath::Pow(InfiniteGearRatio, InGear - MaxGear);
 }
-
-void UMotorPhysicsSimComponent::Reset()
-{
-}
diff --git a/Engine/Plugins/Runtime/AudioMotorSim/Source/AudioMotorSimStandardComponents/Private/ResistanceMotorSimComponent.cpp b/Engine/Plugins/Runtime/AudioMotorSim/Source/AudioMotorSimStandardComponents/Private/ResistanceMotorSimComponent.cpp
--- a/Engine/Plugins/Runtime/AudioMotorSim/Source/AudioMotorSimStandardComponents/Private/ResistanceMotorSimComponent.cpp
+++ b/Engine/Plugins/Runtime/AudioMotorSim/Source/AudioMotorSimStandardComponents/Private/ResistanceMotorSimComponent.cpp
@@ -22,4 +22,6 @@ void UResistanceMotorSimComponent::Update(FAudioMotorSimInputContext& Input, FAu
 		
 		Input.SurfaceFrictionModifier += ZFriction + SideFriction;
 	}
+
+	Super::Update(Input, RuntimeInfo);
 }
diff --git a/Engine/Plugins/Runtime/AudioMotorSim/Source/AudioMotorSimStandardComponents/Private/RevLimiterMotorSimComponent.cpp b/Engine/Plugins/Runtime/AudioMotorSim/Source/AudioMotorSimStandardComponents/Private/RevLimiterMotorSimComponent.cpp
--- a/Engine/Plugins/Runtime/AudioMotorSim/Source/AudioMotorSimStandardComponents/Private/RevLimiterMotorSimComponent.cpp
+++ b/Engine/Plugins/Runtime/AudioMotorSim/Source/AudioMotorSimStandardComponents/Private/RevLimiterMotorSimComponent.cpp
@@ -17,6 +17,8 @@ void URevLimiterMotorSimComponent::Update(FAudioMotorSimInputContext& Input, FAu
 			bActive = false;
 			OnRevLimiterStateChanged.Broadcast(bActive);
 		}
+
+		Super::Update(Input, RuntimeInfo);
 		return;
 	}
 
@@ -47,6 +49,7 @@ void URevLimiterMotorSimComponent::Update(FAudioMotorSimInputContext& Input, FAu
 	if (Input.bGrounded)
 	{
 		TimeInAir = 0.0f;
+		Super::Update(Input, RuntimeInfo);
 		return;
 	}
 
@@ -61,10 +64,14 @@ void URevLimiterMotorSimComponent::Update(FAudioMotorSimInputContext& Input, FAu
 	{
 		Input.Throttle = 0.0f;
 	}
+
+	Super::Update(Input, RuntimeInfo);
 }
 
 void URevLimiterMotorSimComponent::Reset()
 {
+	Super::Reset();
+
 	TimeRemaining = 0.f;
 	TimeInAir = 0.f;
 }
diff --git a/Engine/Plugins/Runtime/AudioMotorSim/Source/AudioMotorSimStandardComponents/Private/ReverseMotorSimComponent.cpp b/Engine/Plugins/Runtime/AudioMotorSim/Source/AudioMotorSimStandardComponents/Private/ReverseMotorSimComponent.cpp
--- a/Engine/Plugins/Runtime/AudioMotorSim/Source/AudioMotorSimStandardComponents/Private/ReverseMotorSimComponent.cpp
+++ b/Engine/Plugins/Runtime/AudioMotorSim/Source/AudioMotorSimStandardComponents/Private/ReverseMotorSimComponent.cpp
@@ -9,9 +9,12 @@ void UReverseMotorSimComponent::Update(FAudioMotorSimInputContext& Input, FAudio
 {
 	if (Input.ForwardSpeed >= 0.f)
 	{
+		Super::Update(Input, RuntimeInfo);
 		return;
 	}
 
 	Input.MotorFrictionModifier *= ReverseEngineResistanceModifier;
+
+	Super::Update(Input, RuntimeInfo);
 }
 
diff --git a/Engine/Plugins/Runtime/AudioMotorSim/Source/AudioMotorSimStandardComponents/Private/RpmCurveMotorSimComponent.cpp b/Engine/Plugins/Runtime/AudioMotorSim/Source/AudioMotorSimStandardComponents/Private/RpmCurveMotorSimComponent.cpp
--- a/Engine/Plugins/Runtime/AudioMotorSim/Source/AudioMotorSimStandardComponents/Private/RpmCurveMotorSimComponent.cpp
+++ b/Engine/Plugins/Runtime/AudioMotorSim/Source/AudioMotorSimStandardComponents/Private/RpmCurveMotorSimComponent.cpp
@@ -9,6 +9,7 @@ void URpmCurveMotorSimComponent::Update(FAudioMotorSimInputContext& Input, FAudi
 {
 	if (Gears.Num() == 0)
 	{
+		Super::Update(Input, RuntimeInfo);
 		return;
 	}
 
@@ -53,6 +54,8 @@ void URpmCurveMotorSimComponent::Update(FAudioMotorSimInputContext& Input, FAudi
 	{
 		RuntimeInfo.Rpm = TargetRpm;
 	}
+
+	Super::Update(Input, RuntimeInfo);
 }
 
 int32 URpmCurveMotorSimComponent::GetDesiredGearForSpeed(const float Speed) const
diff --git a/Engine/Plugins/Runtime/AudioMotorSim/Source/AudioMotorSimStandardComponents/Private/ThrottleStateMotorSimComponent.cpp b/Engine/Plugins/Runtime/AudioMotorSim/Source/AudioMotorSimStandardComponents/Private/ThrottleStateMotorSimComponent.cpp
--- a/Engine/Plugins/Runtime/AudioMotorSim/Source/AudioMotorSimStandardComponents/Private/ThrottleStateMotorSimComponent.cpp
+++ b/Engine/Plugins/Runtime/AudioMotorSim/Source/AudioMotorSimStandardComponents/Private/ThrottleStateMotorSimComponent.cpp
@@ -38,10 +38,14 @@ void UThrottleStateMotorSimComponent::Update(FAudioMotorSimInputContext& Input,
 
     bPrevCarAccelerating = bCarAccelerating;
     bPrevCarIdling = bCarIdling;
+	
+	Super::Update(Input, RuntimeInfo);
 }
 
 void UThrottleStateMotorSimComponent::Reset()
 {
+	Super::Reset();
+
 	FullThrottleTime = 0.f;
 
 	bPrevCarAccelerating = false;
diff --git a/Engine/Plugins/Runtime/AudioMotorSim/Source/AudioMotorSimStandardComponents/Private/VelocitySyncMotorSimComponent.cpp b/Engine/Plugins/Runtime/AudioMotorSim/Source/AudioMotorSimStandardComponents/Private/VelocitySyncMotorSimComponent.cpp
--- a/Engine/Plugins/Runtime/AudioMotorSim/Source/AudioMotorSimStandardComponents/Private/VelocitySyncMotorSimComponent.cpp
+++ b/Engine/Plugins/Runtime/AudioMotorSim/Source/AudioMotorSimStandardComponents/Private/VelocitySyncMotorSimComponent.cpp
@@ -8,6 +8,7 @@ void UVelocitySyncMotorSimComponent::Update(FAudioMotorSimInputContext& Input, F
 {
 	if(Input.bClutchEngaged || !Input.bDriving || !Input.bGrounded)
 	{
+		Super::Update(Input, RuntimeInfo);
 		return;
 	}
 	
@@ -33,9 +34,13 @@ void UVelocitySyncMotorSimComponent::Update(FAudioMotorSimInputContext& Input, F
 		
 		InterpTimeLeft -= Input.DeltaTime;
 	}
+
+	Super::Update(Input, RuntimeInfo);
 }
 void UVelocitySyncMotorSimComponent::Reset()
 {
+	Super::Reset();
+
 	NoThrottleTimeElapsed = 0.f;
 	InterpTimeLeft = 0.f;
 }
diff --git a/Engine/Plugins/Runtime/AudioMotorSim/Source/AudioMotorSimStandardComponents/Public/MotorPhysicsSimComponent.h b/Engine/Plugins/Runtime/AudioMotorSim/Source/AudioMotorSimStandardComponents/Public/MotorPhysicsSimComponent.h
--- a/Engine/Plugins/Runtime/AudioMotorSim/Source/AudioMotorSimStandardComponents/Public/MotorPhysicsSimComponent.h
+++ b/Engine/Plugins/Runtime/AudioMotorSim/Source/AudioMotorSimStandardComponents/Public/MotorPhysicsSimComponent.h
@@ -71,7 +71,6 @@ public:
 	FOnGearChanged OnGearChangedEvent;
 
 	virtual void Update(FAudioMotorSimInputContext& Input, FAudioMotorSimRuntimeContext& RuntimeInfo) override;
-	virtual void Reset() override;
 
 private:
 	float CalcRpm(const float InGearRatio, const float InSpeed) const;
