Set clutch engaged in Motor Sim Rev Limiter when vehicle is not driving

diff --git a/Engine/Plugins/Runtime/AudioMotorSim/Source/AudioMotorSimStandardComponents/Private/RevLimiterMotorSimComponent.cpp b/Engine/Plugins/Runtime/AudioMotorSim/Source/AudioMotorSimStandardComponents/Private/RevLimiterMotorSimComponent.cpp
--- a/Engine/Plugins/Runtime/AudioMotorSim/Source/AudioMotorSimStandardComponents/Private/RevLimiterMotorSimComponent.cpp
+++ b/Engine/Plugins/Runtime/AudioMotorSim/Source/AudioMotorSimStandardComponents/Private/RevLimiterMotorSimComponent.cpp
@@ -45,6 +45,10 @@ void URevLimiterMotorSimComponent::Update(FAudioMotorSimInputContext& Input, FAu
 		TimeRemaining -= Input.DeltaTime;
 		Input.bClutchEngaged = true;
 	}
+	else if (Input.bDriving == false)
+	{
+		Input.bClutchEngaged = true;
+	}
 
 	if (Input.bGrounded)
 	{
