Expose AudioMotorSim Reset function to blueprint, rename BP event to OnReset to avoid conflicts

diff --git a/Engine/Plugins/Runtime/AudioMotorSim/Source/AudioMotorSim/Public/IAudioMotorSim.h b/Engine/Plugins/Runtime/AudioMotorSim/Source/AudioMotorSim/Public/IAudioMotorSim.h
--- a/Engine/Plugins/Runtime/AudioMotorSim/Source/AudioMotorSim/Public/IAudioMotorSim.h
+++ b/Engine/Plugins/Runtime/AudioMotorSim/Source/AudioMotorSim/Public/IAudioMotorSim.h
@@ -8,7 +8,7 @@
 struct FAudioMotorSimInputContext;
 struct FAudioMotorSimRuntimeContext;
 
-UINTERFACE(BlueprintType)
+UINTERFACE(BlueprintType, NotBlueprintable)
 class UAudioMotorSim : public UInterface
 {
 	GENERATED_UINTERFACE_BODY()
@@ -21,6 +21,8 @@ class AUDIOMOTORSIM_API IAudioMotorSim
 public:
 	virtual void Update(FAudioMotorSimInputContext& Input, FAudioMotorSimRuntimeContext& RuntimeInfo) = 0;
 
+	// Use to reset any state that might be desired. Will be called automatically if the entire MotorSim is Reset, or call it manually
+	UFUNCTION(BlueprintCallable, Category = "AudioMotorSim")
 	virtual void Reset() {}
 };
 
@@ -33,7 +35,7 @@ public:
 	UAudioMotorSimComponent(const FObjectInitializer& ObjectInitializer);
 	
 	virtual void Update(FAudioMotorSimInputContext& Input, FAudioMotorSimRuntimeContext& RuntimeInfo) override;
-	
+
 	virtual void Reset() override;
 
 	/* Called every tick that this component is being updated. Use "Set Members in Struct" to update values for future components in the chain. The return value does nothing.
@@ -41,11 +43,11 @@ public:
 	* @param RuntimeInfo	Holds values which are saved between update frames to represent the output or state of the simulation
 	* @return				Vestigial, does nothing.
 	*/
-	UFUNCTION(BlueprintImplementableEvent, Category = "AudioMotorSim", DisplayName = "Update")
+	UFUNCTION(BlueprintImplementableEvent, Category = "AudioMotorSim", DisplayName = "OnUpdate")
 	bool BP_Update(UPARAM(ref) FAudioMotorSimInputContext& Input, UPARAM(ref) FAudioMotorSimRuntimeContext& RuntimeInfo);
 	
-	// Use to reset any state that might be desired. Will be called automatically if the entire MotorSim is Reset, or call it manually
-	UFUNCTION(BlueprintImplementableEvent, BlueprintCallable, Category="AudioMotorSim", DisplayName = "Reset")
+	// Called when something Resets this component
+	UFUNCTION(BlueprintImplementableEvent, BlueprintCallable, Category="AudioMotorSim", DisplayName = "OnReset")
 	void BP_Reset();
 
 #if WITH_EDITORONLY_DATA
