[MotorSim] adding back decel scale, which was missed from the rev limiter component. re-implementing directly on the physics sim component for control.

#rnx

diff --git a/Engine/Plugins/Runtime/AudioMotorSim/Source/AudioMotorSimStandardComponents/Private/MotorPhysicsSimComponent.cpp b/Engine/Plugins/Runtime/AudioMotorSim/Source/AudioMotorSimStandardComponents/Private/MotorPhysicsSimComponent.cpp
--- a/Engine/Plugins/Runtime/AudioMotorSim/Source/AudioMotorSimStandardComponents/Private/MotorPhysicsSimComponent.cpp
+++ b/Engine/Plugins/Runtime/AudioMotorSim/Source/AudioMotorSimStandardComponents/Private/MotorPhysicsSimComponent.cpp
@@ -22,13 +22,14 @@ void UMotorPhysicsSimComponent::Update(FAudioMotorSimInputContext& Input, FAudio
 	const int32 PrevGear = Gear;
 	
 	const float ClutchModifier = Input.bClutchEngaged ? ClutchedForceModifier : 1.f;
-	
+	const float ClutchDecelModifier = Input.bClutchEngaged ? FMath::Max(ClutchedDecelScale * (1.f - Throttle), 0.f) : 1.f;
+
 	//Forces in Newtons.
 	const float CurrentEngineTorque = EngineTorque * Throttle * GearRatio * ClutchModifier;
 	
 	const float BrakingForce = Input.Brake * BrakingHorsePower;
 	const float WindResistance = WindResistancePerVelocity * FMath::Square(SpeedKmh);
-	const float FrictionResistance = Weight * GroundFriction * Input.SurfaceFrictionModifier;
+	const float FrictionResistance = Weight * GroundFriction * Input.SurfaceFrictionModifier * ClutchDecelModifier;
 	const float EngineResistance = CurrentEngineTorque * FMath::Max(Input.MotorFrictionModifier - Throttle, 0.0f) * EngineFriction;
 	
 	const float NetForce = CurrentEngineTorque - BrakingForce - WindResistance - FrictionResistance - EngineResistance;
diff --git a/Engine/Plugins/Runtime/AudioMotorSim/Source/AudioMotorSimStandardComponents/Public/MotorPhysicsSimComponent.h b/Engine/Plugins/Runtime/AudioMotorSim/Source/AudioMotorSimStandardComponents/Public/MotorPhysicsSimComponent.h
--- a/Engine/Plugins/Runtime/AudioMotorSim/Source/AudioMotorSimStandardComponents/Public/MotorPhysicsSimComponent.h
+++ b/Engine/Plugins/Runtime/AudioMotorSim/Source/AudioMotorSimStandardComponents/Public/MotorPhysicsSimComponent.h
@@ -46,6 +46,10 @@ public:
 
 	UPROPERTY(EditAnywhere, BlueprintReadWrite, Category = "Resistance")
 	float ClutchedForceModifier = 1.f;
+
+	// how much more quickly the engine decelerates off-throttle while clutched
+	UPROPERTY(EditAnywhere, BlueprintReadWrite, Category = "Resistance")
+	float ClutchedDecelScale = 1.f;
 	
 	UPROPERTY(EditAnywhere, BlueprintReadWrite, Category = "Resistance")
 	float EngineGearRatio = 50.f;
