Add ability to disable/enable motor sim components at runtime.

#rnx

diff --git a/Engine/Plugins/Runtime/AudioMotorSim/Source/AudioMotorSim/Private/IAudioMotorSim.cpp b/Engine/Plugins/Runtime/AudioMotorSim/Source/AudioMotorSim/Private/IAudioMotorSim.cpp
--- a/Engine/Plugins/Runtime/AudioMotorSim/Source/AudioMotorSim/Private/IAudioMotorSim.cpp
+++ b/Engine/Plugins/Runtime/AudioMotorSim/Source/AudioMotorSim/Private/IAudioMotorSim.cpp
@@ -17,7 +17,10 @@ UAudioMotorSimComponent::UAudioMotorSimComponent(const FObjectInitializer& Objec
 
 void UAudioMotorSimComponent::Update(FAudioMotorSimInputContext& Input, FAudioMotorSimRuntimeContext& RuntimeInfo)
 {
-	BP_Update(Input, RuntimeInfo);
+	if (bEnabled)
+	{
+		BP_Update(Input, RuntimeInfo);
+	}
 
 #if WITH_EDITORONLY_DATA
 	CachedInput = Input;
@@ -35,6 +38,11 @@ void UAudioMotorSimComponent::Reset()
 #endif
 }
 
+void UAudioMotorSimComponent::SetEnabled(bool bNewEnabled)
+{
+	bEnabled = bNewEnabled;
+}
+
 #if WITH_EDITORONLY_DATA
 void UAudioMotorSimComponent::GetCachedData(FAudioMotorSimInputContext& OutInput, FAudioMotorSimRuntimeContext& OutRuntimeInfo)
 {
diff --git a/Engine/Plugins/Runtime/AudioMotorSim/Source/AudioMotorSim/Public/IAudioMotorSim.h b/Engine/Plugins/Runtime/AudioMotorSim/Source/AudioMotorSim/Public/IAudioMotorSim.h
--- a/Engine/Plugins/Runtime/AudioMotorSim/Source/AudioMotorSim/Public/IAudioMotorSim.h
+++ b/Engine/Plugins/Runtime/AudioMotorSim/Source/AudioMotorSim/Public/IAudioMotorSim.h
@@ -49,6 +49,10 @@ public:
 	// Called when something Resets this component
 	UFUNCTION(BlueprintImplementableEvent, BlueprintCallable, Category="AudioMotorSim", DisplayName = "OnReset")
 	void BP_Reset();
+	
+	// controls whether this will run its update function
+	UFUNCTION(BlueprintCallable, Category="AudioMotorSim")
+	void SetEnabled(bool bNewEnabled);
 
 #if WITH_EDITORONLY_DATA
 	// Input data after running this component
@@ -61,4 +65,9 @@ public:
 
 	virtual void GetCachedData(FAudioMotorSimInputContext& OutInput, FAudioMotorSimRuntimeContext& OutRuntimeInfo);
 #endif
+
+protected:
+	// will only update if enabled
+    UPROPERTY(BlueprintReadOnly, Category="AudioMotorSim")
+    bool bEnabled = true;
 };
\ No newline at end of file
