expose IAudioMotorSim enabled state so motormodel component can skip updating it entirely from the interface, instead of relying on each component to implement

#rnx

diff --git a/Engine/Plugins/Runtime/AudioMotorSim/Source/AudioMotorSim/Private/AudioMotorModelComponent.cpp b/Engine/Plugins/Runtime/AudioMotorSim/Source/AudioMotorSim/Private/AudioMotorModelComponent.cpp
--- a/Engine/Plugins/Runtime/AudioMotorSim/Source/AudioMotorSim/Private/AudioMotorModelComponent.cpp
+++ b/Engine/Plugins/Runtime/AudioMotorSim/Source/AudioMotorSim/Private/AudioMotorModelComponent.cpp
@@ -10,7 +10,7 @@ void UAudioMotorModelComponent::Update(FAudioMotorSimInputContext Input)
 {
 	for(FMotorSimEntry& Entry : SimComponents)
 	{
-		if(Entry.Sim)
+		if(Entry.Sim && Entry.Sim->GetEnabled())
 		{
 			Entry.Sim->Update(Input, CachedRuntimeContext);
 		}
diff --git a/Engine/Plugins/Runtime/AudioMotorSim/Source/AudioMotorSim/Public/IAudioMotorSim.h b/Engine/Plugins/Runtime/AudioMotorSim/Source/AudioMotorSim/Public/IAudioMotorSim.h
--- a/Engine/Plugins/Runtime/AudioMotorSim/Source/AudioMotorSim/Public/IAudioMotorSim.h
+++ b/Engine/Plugins/Runtime/AudioMotorSim/Source/AudioMotorSim/Public/IAudioMotorSim.h
@@ -24,6 +24,9 @@ public:
 	// Use to reset any state that might be desired. Will be called automatically if the entire MotorSim is Reset, or call it manually
 	UFUNCTION(BlueprintCallable, Category = "AudioMotorSim")
 	virtual void Reset() {}
+
+	UFUNCTION(BlueprintCallable, Category = "AudioMotorSim")
+	virtual bool GetEnabled() { return false; }
 };
 
 UCLASS(Abstract, Blueprintable, Category = "AudioMotorSim", meta=(BlueprintSpawnableComponent))
@@ -38,6 +41,8 @@ public:
 
 	virtual void Reset() override;
 
+	virtual bool GetEnabled() override { return bEnabled; }
+
 	/* Called every tick that this component is being updated. Use "Set Members in Struct" to update values for future components in the chain. The return value does nothing.
 	* @param Input			Holds values which are not saved between update frames which represent input to the simulation
 	* @param RuntimeInfo	Holds values which are saved between update frames to represent the output or state of the simulation
@@ -66,7 +71,6 @@ public:
 	virtual void GetCachedData(FAudioMotorSimInputContext& OutInput, FAudioMotorSimRuntimeContext& OutRuntimeInfo);
 #endif
 
-protected:
 	// will only update if enabled
     UPROPERTY(BlueprintReadOnly, Category="AudioMotorSim")
     bool bEnabled = true;
