Optimization: Adding const ref & reducing struct copies in AudioMotorModelComponent

diff --git a/Engine/Plugins/Runtime/AudioMotorSim/Source/AudioMotorSim/Private/AudioMotorModelComponent.cpp b/Engine/Plugins/Runtime/AudioMotorSim/Source/AudioMotorSim/Private/AudioMotorModelComponent.cpp
--- a/Engine/Plugins/Runtime/AudioMotorSim/Source/AudioMotorSim/Private/AudioMotorModelComponent.cpp
+++ b/Engine/Plugins/Runtime/AudioMotorSim/Source/AudioMotorSim/Private/AudioMotorModelComponent.cpp
@@ -6,25 +6,25 @@
 
 #include UE_INLINE_GENERATED_CPP_BY_NAME(AudioMotorModelComponent)
 
-void UAudioMotorModelComponent::Update(FAudioMotorSimInputContext Input)
+void UAudioMotorModelComponent::Update(const FAudioMotorSimInputContext& Input)
 {
-	for(FMotorSimEntry& Entry : SimComponents)
+	CachedInputContext = Input;
+	
+	for(const FMotorSimEntry& Entry : SimComponents)
 	{
 		if(Entry.Sim && Entry.Sim->GetEnabled())
 		{
-			Entry.Sim->Update(Input, CachedRuntimeContext);
+			Entry.Sim->Update(CachedInputContext, CachedRuntimeContext);
 		}
 	}
 
-	for(TScriptInterface<IAudioMotorSimOutput> Component : AudioComponents)
+	for(TScriptInterface<IAudioMotorSimOutput>& Component : AudioComponents)
 	{
 		if(Component)
 		{
-			Component->Update(Input, CachedRuntimeContext);
+			Component->Update(CachedInputContext, CachedRuntimeContext);
 		}
 	}
-	
-	CachedInputContext = Input;
 }
 	
 void UAudioMotorModelComponent::Reset()
diff --git a/Engine/Plugins/Runtime/AudioMotorSim/Source/AudioMotorSim/Public/AudioMotorModelComponent.h b/Engine/Plugins/Runtime/AudioMotorSim/Source/AudioMotorSim/Public/AudioMotorModelComponent.h
--- a/Engine/Plugins/Runtime/AudioMotorSim/Source/AudioMotorSim/Public/AudioMotorModelComponent.h
+++ b/Engine/Plugins/Runtime/AudioMotorSim/Source/AudioMotorSim/Public/AudioMotorModelComponent.h
@@ -32,7 +32,7 @@ public:
 	TArray<TScriptInterface<IAudioMotorSimOutput>> AudioComponents;
 
 	UFUNCTION(BlueprintCallable, Category = MotorModel)
-	virtual void Update(FAudioMotorSimInputContext Input);
+	virtual void Update(const FAudioMotorSimInputContext& Input);
 	
 	UFUNCTION(BlueprintCallable, Category = MotorModel)
 	virtual void Reset();
@@ -65,7 +65,7 @@ public:
 	FAudioMotorSimRuntimeContext GetRuntimeInfo() const { return CachedRuntimeContext; }
 
 	UFUNCTION(BlueprintPure, Category = State)
-	FAudioMotorSimInputContext GetCachedInputData() const { return CachedInputContext; }
+	const FAudioMotorSimInputContext& GetCachedInputData() const { return CachedInputContext; }
 
 private:
 	FAudioMotorSimRuntimeContext CachedRuntimeContext;
