Setting up Motor Sims to be able to used in components
-exposed variables to complete component setup.

diff --git a/Engine/Plugins/Runtime/AudioMotorSim/Source/AudioMotorSimStandardComponents/Public/ResistanceMotorSimComponent.h b/Engine/Plugins/Runtime/AudioMotorSim/Source/AudioMotorSimStandardComponents/Public/ResistanceMotorSimComponent.h
--- a/Engine/Plugins/Runtime/AudioMotorSim/Source/AudioMotorSimStandardComponents/Public/ResistanceMotorSimComponent.h
+++ b/Engine/Plugins/Runtime/AudioMotorSim/Source/AudioMotorSimStandardComponents/Public/ResistanceMotorSimComponent.h
@@ -13,15 +13,15 @@ class UResistanceMotorSimComponent : public UAudioMotorSimComponent
 public:
 
 	// How much to increase surface friction when driving straight up. Scales linearly based on driving angle.
-	UPROPERTY(EditAnywhere, Category = "Resistance")
+	UPROPERTY(EditAnywhere, BlueprintReadWrite, Category = "Resistance")
 	float UpSpeedMaxFriction = 1.f;
 
 	// Minimum speed to apply this extra resistance
-	UPROPERTY(EditAnywhere, Category = "Resistance", meta = (ClampMin = "1.0"))
+	UPROPERTY(EditAnywhere, BlueprintReadWrite, Category = "Resistance", meta = (ClampMin = "1.0"))
 	float MinSpeed = 100.f;
 
 	// Additional friction to add based on lateral speed
-	UPROPERTY(EditAnywhere, Category = "Resistance")
+	UPROPERTY(EditAnywhere, BlueprintReadWrite, Category = "Resistance")
 	FRuntimeFloatCurve SideSpeedFrictionCurve;
 	
 	virtual void Update(FAudioMotorSimInputContext& Input, FAudioMotorSimRuntimeContext& RuntimeInfo) override;
diff --git a/Engine/Plugins/Runtime/AudioMotorSim/Source/AudioMotorSimStandardComponents/Public/ReverseMotorSimComponent.h b/Engine/Plugins/Runtime/AudioMotorSim/Source/AudioMotorSimStandardComponents/Public/ReverseMotorSimComponent.h
--- a/Engine/Plugins/Runtime/AudioMotorSim/Source/AudioMotorSimStandardComponents/Public/ReverseMotorSimComponent.h
+++ b/Engine/Plugins/Runtime/AudioMotorSim/Source/AudioMotorSimStandardComponents/Public/ReverseMotorSimComponent.h
@@ -12,7 +12,7 @@ class UReverseMotorSimComponent : public UAudioMotorSimComponent
 	
 public:
 	// How much to scale engine friction by when reversing
-	UPROPERTY(EditAnywhere, Category = "Reverse")
+	UPROPERTY(EditAnywhere, BlueprintReadWrite, Category = "Reverse")
 	float ReverseEngineResistanceModifier = 1.f;
 	
 	virtual void Update(FAudioMotorSimInputContext& Input, FAudioMotorSimRuntimeContext& RuntimeInfo) override;
diff --git a/Engine/Plugins/Runtime/AudioMotorSim/Source/AudioMotorSimStandardComponents/Public/ThrottleStateMotorSimComponent.h b/Engine/Plugins/Runtime/AudioMotorSim/Source/AudioMotorSimStandardComponents/Public/ThrottleStateMotorSimComponent.h
--- a/Engine/Plugins/Runtime/AudioMotorSim/Source/AudioMotorSimStandardComponents/Public/ThrottleStateMotorSimComponent.h
+++ b/Engine/Plugins/Runtime/AudioMotorSim/Source/AudioMotorSimStandardComponents/Public/ThrottleStateMotorSimComponent.h
@@ -24,7 +24,7 @@ public:
 	UPROPERTY(BlueprintAssignable, Category = "ThrottleState")
 	FOnEngineBlowoff OnEngineBlowoff;
 
-	UPROPERTY(EditAnywhere, Category = "ThrottleState")
+	UPROPERTY(EditAnywhere, BlueprintReadWrite, Category = "ThrottleState")
 	float BlowoffMinThrottleTime = 1.f;
 	
 	virtual void Update(FAudioMotorSimInputContext& Input, FAudioMotorSimRuntimeContext& RuntimeInfo) override;
diff --git a/Engine/Plugins/Runtime/AudioMotorSim/Source/AudioMotorSimStandardComponents/Public/VelocitySyncMotorSimComponent.h b/Engine/Plugins/Runtime/AudioMotorSim/Source/AudioMotorSimStandardComponents/Public/VelocitySyncMotorSimComponent.h
--- a/Engine/Plugins/Runtime/AudioMotorSim/Source/AudioMotorSimStandardComponents/Public/VelocitySyncMotorSimComponent.h
+++ b/Engine/Plugins/Runtime/AudioMotorSim/Source/AudioMotorSimStandardComponents/Public/VelocitySyncMotorSimComponent.h
@@ -19,7 +19,7 @@ public:
 	UPROPERTY(EditAnywhere, BlueprintReadWrite, Category = "VelocitySync")
 	float SpeedThreshold = 200.f;
 
-	UPROPERTY(EditAnywhere, Category = "VelocitySync")
+	UPROPERTY(EditAnywhere, BlueprintReadWrite, Category = "VelocitySync")
 	FRuntimeFloatCurve SpeedToRpmCurve;
 
 	UPROPERTY(EditAnywhere, BlueprintReadWrite, Category = "VelocitySync")
