Add Debug Print to AudioMotorModelComponent

diff --git a/Engine/Plugins/Runtime/AudioMotorSim/Source/AudioMotorSim/Private/AudioMotorModelComponent.cpp b/Engine/Plugins/Runtime/AudioMotorSim/Source/AudioMotorSim/Private/AudioMotorModelComponent.cpp
--- a/Engine/Plugins/Runtime/AudioMotorSim/Source/AudioMotorSim/Private/AudioMotorModelComponent.cpp
+++ b/Engine/Plugins/Runtime/AudioMotorSim/Source/AudioMotorSim/Private/AudioMotorModelComponent.cpp
@@ -1,11 +1,55 @@
 // Copyright Epic Games, Inc. All Rights Reserved.
 
 #include "AudioMotorModelComponent.h"
+#include "Engine.h"
 #include "IAudioMotorSim.h"
 #include "IAudioMotorSimOutput.h"
+#include "Math/Color.h"
 
 #include UE_INLINE_GENERATED_CPP_BY_NAME(AudioMotorModelComponent)
 
+namespace AudioMotorModelComponentPrivate	
+{
+	template <typename ParamType>
+	void DebugPrintParam(const FName ParamName, ParamType Value, const FColor& TextColor)
+	{
+		static_assert(
+		std::is_same_v<ParamType, float> ||
+		std::is_same_v<ParamType, int32> ||
+		std::is_same_v<ParamType, bool>,
+		"Unsupported parameter type.");
+
+		if(!GEngine)
+		{
+			return;
+		}
+		
+
+		FString ParamNameAsString = ParamName.ToString();
+		constexpr float TimeToDisplay = 0.f;
+		uint64 DebugKey = GetTypeHash(ParamName.GetComparisonIndex()) + ParamName.GetNumber();
+
+		
+		// NOTE: This currently only supports the parameter types currently being used.
+		// New cases will need to be added for new parameter types.
+		if constexpr (std::is_same_v<ParamType, float>)
+		{
+			const FString Message = FString::Printf(TEXT("%s: %f"), *ParamNameAsString, Value);
+			GEngine->AddOnScreenDebugMessage(DebugKey, TimeToDisplay, TextColor, Message);
+		}
+		else if constexpr (std::is_same_v<ParamType, int32>)
+		{
+			const FString Message = FString::Printf(TEXT("%s: %d"), *ParamNameAsString, Value);
+			GEngine->AddOnScreenDebugMessage(DebugKey, TimeToDisplay, TextColor, Message);
+		}
+		else if constexpr (std::is_same_v<ParamType, bool>)
+		{
+			const FString Message = FString::Printf(TEXT("%s: %s"), *ParamNameAsString, Value ? TEXT("True") : TEXT("False"));
+			GEngine->AddOnScreenDebugMessage(DebugKey, TimeToDisplay, TextColor, Message);
+		}
+	}
+}
+
 void UAudioMotorModelComponent::Update(const FAudioMotorSimInputContext& Input)
 {
 	CachedInputContext = Input;
@@ -110,3 +154,55 @@ void UAudioMotorModelComponent::RemoveMotorSimComponent(TScriptInterface<IAudioM
 		}
 	}
 }
+
+#if !(UE_BUILD_SHIPPING || UE_BUILD_TEST)
+void UAudioMotorModelComponent::DebugPrintRuntimeInfo() const
+{
+	if(GEngine)
+	{
+		return;
+	}
+	
+	const FColor TextColor = FColor::Red;
+	
+	AudioMotorModelComponentPrivate::DebugPrintParam("RPM", CachedRuntimeContext.Rpm, TextColor);
+	AudioMotorModelComponentPrivate::DebugPrintParam("Gear", CachedRuntimeContext.Gear, TextColor);
+	AudioMotorModelComponentPrivate::DebugPrintParam("Volume", CachedRuntimeContext.Volume, TextColor);
+	AudioMotorModelComponentPrivate::DebugPrintParam("Is Shifting", CachedRuntimeContext.bShifting, TextColor);
+	AudioMotorModelComponentPrivate::DebugPrintParam("Pitch", CachedRuntimeContext.Pitch, TextColor);
+	
+	const uint64 HeaderDebugKey = PointerHash("DebugPrintRuntimeInfo");
+	FString HeaderDebugMessage = FString::Printf(TEXT("=== [%s] Audio Motor Model Component Runtime Info ==="), *this->GetFName().ToString());
+	GEngine->AddOnScreenDebugMessage(HeaderDebugKey, 0.f,  TextColor, HeaderDebugMessage);
+}
+
+void UAudioMotorModelComponent::DebugPrintInputContext() const
+{
+	if(GEngine)
+	{
+		return;
+	}
+	
+	const FColor TextColor = FColor::Red;
+	
+	AudioMotorModelComponentPrivate::DebugPrintParam("DeltaTime", CachedInputContext.DeltaTime, TextColor);
+	AudioMotorModelComponentPrivate::DebugPrintParam("Throttle", CachedInputContext.Throttle, TextColor);
+	AudioMotorModelComponentPrivate::DebugPrintParam("Brake", CachedInputContext.Brake, TextColor);
+	AudioMotorModelComponentPrivate::DebugPrintParam("Speed", CachedInputContext.Speed, TextColor);
+	AudioMotorModelComponentPrivate::DebugPrintParam("ForwardSpeed", CachedInputContext.ForwardSpeed, TextColor);
+	AudioMotorModelComponentPrivate::DebugPrintParam("SideSpeed", CachedInputContext.SideSpeed, TextColor);
+	AudioMotorModelComponentPrivate::DebugPrintParam("UpSpeed", CachedInputContext.UpSpeed, TextColor);
+	AudioMotorModelComponentPrivate::DebugPrintParam("SurfaceFrictionModifier", CachedInputContext.SurfaceFrictionModifier, TextColor);
+	AudioMotorModelComponentPrivate::DebugPrintParam("MotorFrictionModifier", CachedInputContext.MotorFrictionModifier, TextColor);
+	AudioMotorModelComponentPrivate::DebugPrintParam("Boost", CachedInputContext.Boost, TextColor);
+	AudioMotorModelComponentPrivate::DebugPrintParam("bDriving", CachedInputContext.bDriving, TextColor);
+	AudioMotorModelComponentPrivate::DebugPrintParam("bGrounded", CachedInputContext.bGrounded, TextColor);
+	AudioMotorModelComponentPrivate::DebugPrintParam("bCanShift", CachedInputContext.bCanShift, TextColor);
+	AudioMotorModelComponentPrivate::DebugPrintParam("bClutchEngaged", CachedInputContext.bClutchEngaged, TextColor);
+	
+	const uint64 HeaderDebugKey = PointerHash("DebugPrintInputContext");
+	FString HeaderDebugMessage = FString::Printf(TEXT("=== [%s] Audio Motor Model Input Context ==="), *this->GetFName().ToString());
+	GEngine->AddOnScreenDebugMessage(HeaderDebugKey, 0.f,  TextColor, HeaderDebugMessage);
+}
+
+#endif
\ No newline at end of file
diff --git a/Engine/Plugins/Runtime/AudioMotorSim/Source/AudioMotorSim/Public/AudioMotorModelComponent.h b/Engine/Plugins/Runtime/AudioMotorSim/Source/AudioMotorSim/Public/AudioMotorModelComponent.h
--- a/Engine/Plugins/Runtime/AudioMotorSim/Source/AudioMotorSim/Public/AudioMotorModelComponent.h
+++ b/Engine/Plugins/Runtime/AudioMotorSim/Source/AudioMotorSim/Public/AudioMotorModelComponent.h
@@ -67,6 +67,11 @@ public:
 	UFUNCTION(BlueprintPure, Category = State)
 	const FAudioMotorSimInputContext& GetCachedInputData() const { return CachedInputContext; }
 
+#if !(UE_BUILD_SHIPPING || UE_BUILD_TEST)
+	void DebugPrintRuntimeInfo() const;
+	void DebugPrintInputContext() const;
+#endif
+	
 private:
 	FAudioMotorSimRuntimeContext CachedRuntimeContext;
 	
