Add Debug Print to AudioMotorModelComponent (without monolithic dependency)

diff --git a/Engine/Plugins/Runtime/AudioMotorSim/Source/AudioMotorSim/Private/AudioMotorModelComponent.cpp b/Engine/Plugins/Runtime/AudioMotorSim/Source/AudioMotorSim/Private/AudioMotorModelComponent.cpp
--- a/Engine/Plugins/Runtime/AudioMotorSim/Source/AudioMotorSim/Private/AudioMotorModelComponent.cpp
+++ b/Engine/Plugins/Runtime/AudioMotorSim/Source/AudioMotorSim/Private/AudioMotorModelComponent.cpp
@@ -1,11 +1,73 @@
 // Copyright Epic Games, Inc. All Rights Reserved.
 
 #include "AudioMotorModelComponent.h"
-#include "IAudioMotorSim.h"
 #include "IAudioMotorSimOutput.h"
+#include "IAudioMotorSim.h"
+#include "Kismet/KismetSystemLibrary.h"
+#include "Math/Color.h"
 
 #include UE_INLINE_GENERATED_CPP_BY_NAME(AudioMotorModelComponent)
 
+namespace AudioMotorModelComponentPrivate	
+{
+	template <typename ParamType>
+	void DebugPrintParam(const FName ParamName, ParamType Value, const FColor& TextColor, const UObject* WorldContextObject)
+	{
+		static_assert(
+		std::is_same_v<ParamType, float> ||
+		std::is_same_v<ParamType, int32> ||
+		std::is_same_v<ParamType, bool>,
+		"Unsupported parameter type.");
+				
+		FString ParamNameAsString = ParamName.ToString();
+		FString Message;
+		
+		// NOTE: This currently only supports the parameter types currently being used.
+		// New cases will need to be added for new parameter types.
+		if constexpr (std::is_same_v<ParamType, float>)
+		{
+			Message = FString::Printf(TEXT("%s: %f"), *ParamNameAsString, Value);
+		}
+		else if constexpr (std::is_same_v<ParamType, int32>)
+		{
+			Message = FString::Printf(TEXT("%s: %d"), *ParamNameAsString, Value);
+		}
+		else if constexpr (std::is_same_v<ParamType, bool>)
+		{
+			Message = FString::Printf(TEXT("%s: %s"), *ParamNameAsString, Value ? TEXT("True") : TEXT("False"));
+		}
+
+		constexpr bool bPrintToScreen = true;
+		constexpr bool bPrintToLog = false;
+		constexpr float TimeToDisplay = 0.f;
+		
+		UKismetSystemLibrary::PrintString(
+			WorldContextObject,
+			Message,
+			bPrintToScreen,
+			bPrintToLog,
+			TextColor,
+			TimeToDisplay,
+			ParamName);
+	}
+	
+    void DebugPrintString(const FString& Message, const FColor& TextColor, const UObject* WorldContextObject)
+    { 
+    	constexpr bool bPrintToScreen = true;
+    	constexpr bool bPrintToLog = false;
+    	constexpr float TimeToDisplay = 0.f;
+    	
+    	UKismetSystemLibrary::PrintString(
+    		WorldContextObject,
+    		Message,
+			bPrintToScreen,
+    		bPrintToLog,
+    		TextColor,
+    		TimeToDisplay,
+    		FName(Message));
+    }
+}
+
 void UAudioMotorModelComponent::Update(const FAudioMotorSimInputContext& Input)
 {
 	CachedInputContext = Input;
@@ -110,3 +172,42 @@ void UAudioMotorModelComponent::RemoveMotorSimComponent(TScriptInterface<IAudioM
 		}
 	}
 }
+
+#if !(UE_BUILD_SHIPPING || UE_BUILD_TEST)
+void UAudioMotorModelComponent::DebugPrintRuntimeInfo() const
+{
+	const FColor TextColor = FColor::Red;
+	
+	AudioMotorModelComponentPrivate::DebugPrintParam("RPM", CachedRuntimeContext.Rpm, TextColor, this);
+	AudioMotorModelComponentPrivate::DebugPrintParam("Gear", CachedRuntimeContext.Gear, TextColor, this);
+	AudioMotorModelComponentPrivate::DebugPrintParam("Volume", CachedRuntimeContext.Volume, TextColor, this);
+	AudioMotorModelComponentPrivate::DebugPrintParam("Is Shifting", CachedRuntimeContext.bShifting, TextColor, this);
+	AudioMotorModelComponentPrivate::DebugPrintParam("Pitch", CachedRuntimeContext.Pitch, TextColor, this);
+	
+	FString HeaderDebugMessage = FString::Printf(TEXT("=== [%s] Audio Motor Model Component Runtime Info ==="), *this->GetFName().ToString());
+	AudioMotorModelComponentPrivate::DebugPrintString(HeaderDebugMessage, TextColor, this);
+}
+
+void UAudioMotorModelComponent::DebugPrintInputContext() const
+{
+	const FColor TextColor = FColor::Red;
+	
+	AudioMotorModelComponentPrivate::DebugPrintParam("DeltaTime", CachedInputContext.DeltaTime, TextColor, this);
+	AudioMotorModelComponentPrivate::DebugPrintParam("Throttle", CachedInputContext.Throttle, TextColor, this);
+	AudioMotorModelComponentPrivate::DebugPrintParam("Brake", CachedInputContext.Brake, TextColor, this);
+	AudioMotorModelComponentPrivate::DebugPrintParam("Speed", CachedInputContext.Speed, TextColor, this);
+	AudioMotorModelComponentPrivate::DebugPrintParam("ForwardSpeed", CachedInputContext.ForwardSpeed, TextColor, this);
+	AudioMotorModelComponentPrivate::DebugPrintParam("SideSpeed", CachedInputContext.SideSpeed, TextColor, this);
+	AudioMotorModelComponentPrivate::DebugPrintParam("UpSpeed", CachedInputContext.UpSpeed, TextColor, this);
+	AudioMotorModelComponentPrivate::DebugPrintParam("SurfaceFrictionModifier", CachedInputContext.SurfaceFrictionModifier, TextColor, this);
+	AudioMotorModelComponentPrivate::DebugPrintParam("MotorFrictionModifier", CachedInputContext.MotorFrictionModifier, TextColor, this);
+	AudioMotorModelComponentPrivate::DebugPrintParam("Boost", CachedInputContext.Boost, TextColor, this);
+	AudioMotorModelComponentPrivate::DebugPrintParam("bDriving", CachedInputContext.bDriving, TextColor, this);
+	AudioMotorModelComponentPrivate::DebugPrintParam("bGrounded", CachedInputContext.bGrounded, TextColor, this);
+	AudioMotorModelComponentPrivate::DebugPrintParam("bCanShift", CachedInputContext.bCanShift, TextColor, this);
+	AudioMotorModelComponentPrivate::DebugPrintParam("bClutchEngaged", CachedInputContext.bClutchEngaged, TextColor, this);
+	
+	FString HeaderDebugMessage = FString::Printf(TEXT("=== [%s] Audio Motor Model Input Context ==="), *this->GetFName().ToString());
+	AudioMotorModelComponentPrivate::DebugPrintString(HeaderDebugMessage, TextColor, this);
+}
+#endif
\ No newline at end of file
diff --git a/Engine/Plugins/Runtime/AudioMotorSim/Source/AudioMotorSim/Public/AudioMotorModelComponent.h b/Engine/Plugins/Runtime/AudioMotorSim/Source/AudioMotorSim/Public/AudioMotorModelComponent.h
--- a/Engine/Plugins/Runtime/AudioMotorSim/Source/AudioMotorSim/Public/AudioMotorModelComponent.h
+++ b/Engine/Plugins/Runtime/AudioMotorSim/Source/AudioMotorSim/Public/AudioMotorModelComponent.h
@@ -67,6 +67,11 @@ public:
 	UFUNCTION(BlueprintPure, Category = State)
 	const FAudioMotorSimInputContext& GetCachedInputData() const { return CachedInputContext; }
 
+#if !(UE_BUILD_SHIPPING || UE_BUILD_TEST)
+	void DebugPrintRuntimeInfo() const;
+	void DebugPrintInputContext() const;
+#endif
+	
 private:
 	FAudioMotorSimRuntimeContext CachedRuntimeContext;
 	
