add separate motor model logging for multiplayer

#rnx

diff --git a/Engine/Plugins/Runtime/AudioMotorSim/Source/AudioMotorSim/Private/AudioMotorModelComponent.cpp b/Engine/Plugins/Runtime/AudioMotorSim/Source/AudioMotorSim/Private/AudioMotorModelComponent.cpp
--- a/Engine/Plugins/Runtime/AudioMotorSim/Source/AudioMotorSim/Private/AudioMotorModelComponent.cpp
+++ b/Engine/Plugins/Runtime/AudioMotorSim/Source/AudioMotorSim/Private/AudioMotorModelComponent.cpp
@@ -19,27 +19,32 @@ namespace AudioMotorModelComponentPrivate
 		std::is_same_v<ParamType, bool>,
 		"Unsupported parameter type.");
 				
-		FString ParamNameAsString = ParamName.ToString();
+		const FString ParamNameAsString = ParamName.ToString();
 		FString Message;
+
+		const UObject* Owner = WorldContextObject ? WorldContextObject->GetOuter() : WorldContextObject;
+		const FString ContextObjectName = Owner ? Owner->GetName() : TEXT("");
 		
 		// NOTE: This currently only supports the parameter types currently being used.
 		// New cases will need to be added for new parameter types.
 		if constexpr (std::is_same_v<ParamType, float>)
 		{
-			Message = FString::Printf(TEXT("%s: %f"), *ParamNameAsString, Value);
+			Message = FString::Printf(TEXT("[%s] %s: %f"), *ContextObjectName, *ParamNameAsString, Value);
 		}
 		else if constexpr (std::is_same_v<ParamType, int32>)
 		{
-			Message = FString::Printf(TEXT("%s: %d"), *ParamNameAsString, Value);
+			Message = FString::Printf(TEXT("[%s] %s: %d"), *ContextObjectName, *ParamNameAsString, Value);
 		}
 		else if constexpr (std::is_same_v<ParamType, bool>)
 		{
-			Message = FString::Printf(TEXT("%s: %s"), *ParamNameAsString, Value ? TEXT("True") : TEXT("False"));
+			Message = FString::Printf(TEXT("[%s] %s: %s"), *ContextObjectName, *ParamNameAsString, Value ? TEXT("True") : TEXT("False"));
 		}
 
 		constexpr bool bPrintToScreen = true;
 		constexpr bool bPrintToLog = false;
 		constexpr float TimeToDisplay = 0.f;
+
+		FName DebugKey = FName(FString::Printf(TEXT("%s_%s"), *ContextObjectName, *ParamNameAsString));
 		
 		UKismetSystemLibrary::PrintString(
 			WorldContextObject,
@@ -48,7 +53,7 @@ namespace AudioMotorModelComponentPrivate
 			bPrintToLog,
 			TextColor,
 			TimeToDisplay,
-			ParamName);
+			DebugKey);
 	}
 	
     void DebugPrintString(const FString& Message, const FColor& TextColor, const UObject* WorldContextObject)
@@ -56,15 +61,19 @@ namespace AudioMotorModelComponentPrivate
     	constexpr bool bPrintToScreen = true;
     	constexpr bool bPrintToLog = false;
     	constexpr float TimeToDisplay = 0.f;
+
+		const UObject* Owner = WorldContextObject ? WorldContextObject->GetOuter() : WorldContextObject;
+		const FString ContextObjectName = Owner ? Owner->GetName() : TEXT("");
+		FString StringToPrint = FString::Printf(TEXT("[%s] %s"), *ContextObjectName, *Message);
     	
     	UKismetSystemLibrary::PrintString(
     		WorldContextObject,
-    		Message,
+    		StringToPrint,
 			bPrintToScreen,
     		bPrintToLog,
     		TextColor,
     		TimeToDisplay,
-    		FName(Message));
+    		FName(StringToPrint));
     }
 }
 
