add logs around motorphysiscs sim component to make it easier to tune

diff --git a/Engine/Plugins/Runtime/AudioMotorSim/Source/AudioMotorSim/Private/AudioMotorSimModule.cpp b/Engine/Plugins/Runtime/AudioMotorSim/Source/AudioMotorSim/Private/AudioMotorSimModule.cpp
--- a/Engine/Plugins/Runtime/AudioMotorSim/Source/AudioMotorSim/Private/AudioMotorSimModule.cpp
+++ b/Engine/Plugins/Runtime/AudioMotorSim/Source/AudioMotorSim/Private/AudioMotorSimModule.cpp
@@ -1,6 +1,7 @@
 // Copyright Epic Games, Inc. All Rights Reserved.
 
 #include "AudioMotorSimModule.h"
+#include "AudioMotorSimLogs.h"
 #include "Modules/ModuleManager.h"
 
 class FAudioMotorSimModule : public IModuleInterface
@@ -11,3 +12,5 @@ public:
 };
 
 IMPLEMENT_MODULE(FAudioMotorSimModule, AudioMotorSim);
+
+DEFINE_LOG_CATEGORY(LogAudioMotorSim)
\ No newline at end of file
diff --git a/Engine/Plugins/Runtime/AudioMotorSim/Source/AudioMotorSim/Public/AudioMotorSimLogs.h b/Engine/Plugins/Runtime/AudioMotorSim/Source/AudioMotorSim/Public/AudioMotorSimLogs.h
new file mode 100644
--- /dev/null
+++ b/Engine/Plugins/Runtime/AudioMotorSim/Source/AudioMotorSim/Public/AudioMotorSimLogs.h
@@ -0,0 +1,7 @@
+﻿// Copyright Epic Games, Inc. All Rights Reserved.
+
+#pragma once
+
+#include "Logging/LogMacros.h"
+
+AUDIOMOTORSIM_API DECLARE_LOG_CATEGORY_EXTERN(LogAudioMotorSim, Log, All);
\ No newline at end of file
diff --git a/Engine/Plugins/Runtime/AudioMotorSim/Source/AudioMotorSimStandardComponents/Private/MotorPhysicsSimComponent.cpp b/Engine/Plugins/Runtime/AudioMotorSim/Source/AudioMotorSimStandardComponents/Private/MotorPhysicsSimComponent.cpp
--- a/Engine/Plugins/Runtime/AudioMotorSim/Source/AudioMotorSimStandardComponents/Private/MotorPhysicsSimComponent.cpp
+++ b/Engine/Plugins/Runtime/AudioMotorSim/Source/AudioMotorSimStandardComponents/Private/MotorPhysicsSimComponent.cpp
@@ -1,10 +1,10 @@
 // Copyright Epic Games, Inc. All Rights Reserved.
 
 #include "MotorPhysicsSimComponent.h"
+#include "AudioMotorSimLogs.h"
 #include "AudioMotorSimTypes.h"
 
 #include UE_INLINE_GENERATED_CPP_BY_NAME(MotorPhysicsSimComponent)
-
 void UMotorPhysicsSimComponent::Update(FAudioMotorSimInputContext& Input, FAudioMotorSimRuntimeContext& RuntimeInfo)
 {
 	if(ensure(Weight > 0.f && EngineGearRatio > 0.f) == false)
@@ -34,6 +34,8 @@ void UMotorPhysicsSimComponent::Update(FAudioMotorSimInputContext& Input, FAudio
 	
 	const float NetForce = CurrentEngineTorque - BrakingForce - WindResistance - FrictionResistance - EngineResistance;
 
+	UE_LOG(LogAudioMotorSim, Verbose, TEXT("[UMotorPhysicsSimComponent] NetForce: %f (CurrentEngineTorque: %f - BrakingForce: %f - WindResistance %f - FrictionResistance %f - EngineResistance %f"), NetForce, CurrentEngineTorque, BrakingForce, WindResistance, FrictionResistance, EngineResistance)
+
 	// Convert Rpm to velocity to apply forces
 	float VirtualSpeedKmh =  CalcVelocity(GearRatio, RuntimeInfo.Rpm);
 	VirtualSpeedKmh += ((NetForce / Weight) * Input.DeltaTime);
@@ -42,6 +44,9 @@ void UMotorPhysicsSimComponent::Update(FAudioMotorSimInputContext& Input, FAudio
 	// Now working backwards, what is the engine RPM as a results of this velocity and gear ratio.
 	float Rpm = CalcRpm(GearRatio, VirtualSpeedKmh);
 
+	UE_LOG(LogAudioMotorSim, Verbose, TEXT("[UMotorPhysicsSimComponent] Curr Gear %d - RPM: %f"), Gear, Rpm);
+	UE_LOG(LogAudioMotorSim, Verbose, TEXT("[UMotorPhysicsSimComponent] Prev Gear %d - RPM: %f"), (Gear >= 1 ? Gear - 1 : 0), CalcRpm(GetGearRatio((Gear >= 1 ? Gear - 1 : 0), Input.bClutchEngaged), VirtualSpeedKmh));
+	
 	// Now determine if we need to shift.
 	if (Input.bCanShift && (Rpm > UpShiftMaxRpm))
 	{
