[VehicleAudio] Add revving to preview motor model

#changelist validated
#virtualized

diff --git a/Engine/Plugins/Runtime/AudioMotorSim/Source/AudioMotorSimStandardComponents/Private/RevLimiterMotorSimComponent.cpp b/Engine/Plugins/Runtime/AudioMotorSim/Source/AudioMotorSimStandardComponents/Private/RevLimiterMotorSimComponent.cpp
--- a/Engine/Plugins/Runtime/AudioMotorSim/Source/AudioMotorSimStandardComponents/Private/RevLimiterMotorSimComponent.cpp
+++ b/Engine/Plugins/Runtime/AudioMotorSim/Source/AudioMotorSimStandardComponents/Private/RevLimiterMotorSimComponent.cpp
@@ -5,8 +5,35 @@
 
 #include UE_INLINE_GENERATED_CPP_BY_NAME(RevLimiterMotorSimComponent)
 
+void URevLimiterMotorSimComponent::RevLimitToTarget(FAudioMotorSimInputContext& Input, FAudioMotorSimRuntimeContext& RuntimeInfo)
+{
+	if (RuntimeInfo.Rpm >= LimiterMaxRpm && !bActive)
+	{
+		bActive = true;
+		TimeRemaining = LimitTime;
+		OnRevLimiterHit.Broadcast();
+		OnRevLimiterStateChanged.Broadcast(bActive);
+	}
+
+	if(RuntimeInfo.Rpm <= RecoverRPM && bActive)
+	{
+		bActive = false;
+		OnRevLimiterStateChanged.Broadcast(bActive);
+
+	}
+		
+	Input.Throttle = bActive ? 0.0f : Input.Throttle;
+}
+
 void URevLimiterMotorSimComponent::Update(FAudioMotorSimInputContext& Input, FAudioMotorSimRuntimeContext& RuntimeInfo)
 {
+	
+	if(bRevLimitOnClutchEngaged && Input.bClutchEngaged)
+	{
+		RevLimitToTarget(Input, RuntimeInfo);
+		return;
+	}
+	
 	if (Input.bDriving && Input.bGrounded && Input.SideSpeed < SideSpeedThreshold)
 	{
 		TimeRemaining = 0.f;
diff --git a/Engine/Plugins/Runtime/AudioMotorSim/Source/AudioMotorSimStandardComponents/Public/RevLimiterMotorSimComponent.h b/Engine/Plugins/Runtime/AudioMotorSim/Source/AudioMotorSimStandardComponents/Public/RevLimiterMotorSimComponent.h
--- a/Engine/Plugins/Runtime/AudioMotorSim/Source/AudioMotorSimStandardComponents/Public/RevLimiterMotorSimComponent.h
+++ b/Engine/Plugins/Runtime/AudioMotorSim/Source/AudioMotorSimStandardComponents/Public/RevLimiterMotorSimComponent.h
@@ -16,7 +16,7 @@ public:
 
 	UPROPERTY(EditAnywhere, BlueprintReadWrite, Category = "RevLimiter")
 	float LimitTime = 0.f;
-	
+
 	UPROPERTY(EditAnywhere, BlueprintReadWrite, Category = "RevLimiter")
 	float DecelScale = 0.f;
 	
@@ -28,6 +28,12 @@ public:
 
 	UPROPERTY(EditAnywhere, BlueprintReadWrite, Category = "RevLimiter")
 	float LimiterMaxRpm = 0.f;
+	
+	UPROPERTY(EditAnywhere, BlueprintReadWrite, Category = "RevLimiter")
+	bool bRevLimitOnClutchEngaged = false;
+
+	UPROPERTY(EditAnywhere, BlueprintReadWrite, Category = "RevLimiter")
+	float RecoverRPM = 0.f;
 
 	UPROPERTY(BlueprintAssignable)
 	FOnRevLimiterHit OnRevLimiterHit;
@@ -36,6 +42,8 @@ public:
 	FOnRevLimiterStateChanged OnRevLimiterStateChanged;
 
 private:
+	void RevLimitToTarget(FAudioMotorSimInputContext& Input, FAudioMotorSimRuntimeContext& RuntimeInfo);
+	
 	// Time remaining where the limiter is forcing throttle down
 	float TimeRemaining = 0.f;
 	float TimeInAir = 0.f;
