Allow runtime reconfiguration of AudioMotorModelComponents for systems like cosmetics - remove TInstancedStruct to avoid template issues

#rnx

diff --git a/Engine/Plugins/Runtime/AudioMotorSim/Source/AudioMotorSim/Private/AudioMotorModelComponent.cpp b/Engine/Plugins/Runtime/AudioMotorSim/Source/AudioMotorSim/Private/AudioMotorModelComponent.cpp
--- a/Engine/Plugins/Runtime/AudioMotorSim/Source/AudioMotorSim/Private/AudioMotorModelComponent.cpp
+++ b/Engine/Plugins/Runtime/AudioMotorSim/Source/AudioMotorSim/Private/AudioMotorModelComponent.cpp
@@ -219,6 +219,20 @@ void UAudioMotorModelComponent::RemoveMotorSimComponent(TScriptInterface<IAudioM
 	}
 }
 
+void UAudioMotorModelComponent::ConfigureMotorSimComponents(const TArray<FInstancedStruct>& InConfigData)
+{
+	for(const FInstancedStruct& ConfigData : InConfigData)
+	{
+		for(FMotorSimEntry& SimComponent : SimComponents)
+		{
+			if(SimComponent.Sim)
+			{
+				SimComponent.Sim->ConfigMotorSim(ConfigData);
+			}
+		}
+	}
+}
+
 void UAudioMotorModelComponent::RemoveAllMotorSimComponents()
 {
 	StopOutput();
diff --git a/Engine/Plugins/Runtime/AudioMotorSim/Source/AudioMotorSim/Private/IAudioMotorSim.cpp b/Engine/Plugins/Runtime/AudioMotorSim/Source/AudioMotorSim/Private/IAudioMotorSim.cpp
--- a/Engine/Plugins/Runtime/AudioMotorSim/Source/AudioMotorSim/Private/IAudioMotorSim.cpp
+++ b/Engine/Plugins/Runtime/AudioMotorSim/Source/AudioMotorSim/Private/IAudioMotorSim.cpp
@@ -2,6 +2,8 @@
 
 #include "IAudioMotorSim.h"
 
+#include "AudioMotorSimConfigData.h"
+
 #include UE_INLINE_GENERATED_CPP_BY_NAME(IAudioMotorSim)
 
 UAudioMotorSim::UAudioMotorSim(const FObjectInitializer& ObjectInitializer)
@@ -56,6 +58,13 @@ void UAudioMotorSimComponent::SetEnabled(bool bNewEnabled)
 	bEnabled = bNewEnabled;
 }
 
+void UAudioMotorSimComponent::ConfigMotorSim(const FInstancedStruct& InConfigData)
+{
+	IAudioMotorSim::ConfigMotorSim(InConfigData);
+
+	ensureMsgf(InConfigData.GetPtr<FAudioMotorSimConfigData>(), TEXT("Expected instance struct of being a FAudioMotorSimConfigData type"));
+}
+
 #if WITH_EDITORONLY_DATA
 void UAudioMotorSimComponent::GetCachedData(FAudioMotorSimInputContext& OutInput, FAudioMotorSimRuntimeContext& OutRuntimeInfo)
 {
diff --git a/Engine/Plugins/Runtime/AudioMotorSim/Source/AudioMotorSim/Public/AudioMotorModelComponent.h b/Engine/Plugins/Runtime/AudioMotorSim/Source/AudioMotorSim/Public/AudioMotorModelComponent.h
--- a/Engine/Plugins/Runtime/AudioMotorSim/Source/AudioMotorSim/Public/AudioMotorModelComponent.h
+++ b/Engine/Plugins/Runtime/AudioMotorSim/Source/AudioMotorSim/Public/AudioMotorModelComponent.h
@@ -2,7 +2,10 @@
 #pragma once
 
 #include "Components/ActorComponent.h"
+#include "AudioMotorSimConfigData.h"
 #include "AudioMotorSimTypes.h"
+#include "StructUtils/InstancedStruct.h"
+
 #include "AudioMotorModelComponent.generated.h"
 
 class IAudioMotorSim;
@@ -54,9 +57,12 @@ public:
 	
 	UFUNCTION(BlueprintCallable, Category = MotorModel)
 	void RemoveMotorSimComponent(TScriptInterface<IAudioMotorSim> InComponent);
-
+	
 	UFUNCTION(BlueprintCallable, Category = MotorModel)
 	void RemoveAllMotorSimComponents();
+
+	UFUNCTION(BlueprintCallable, Category = MotorModel)
+	void ConfigureMotorSimComponents(const TArray<FInstancedStruct>& InConfigData);
 	
 	UFUNCTION(BlueprintPure, Category = State)
 	float GetRpm() const { return CachedRuntimeContext.Rpm; }
diff --git a/Engine/Plugins/Runtime/AudioMotorSim/Source/AudioMotorSim/Public/AudioMotorSimConfigData.h b/Engine/Plugins/Runtime/AudioMotorSim/Source/AudioMotorSim/Public/AudioMotorSimConfigData.h
new file mode 100644
--- /dev/null
+++ b/Engine/Plugins/Runtime/AudioMotorSim/Source/AudioMotorSim/Public/AudioMotorSimConfigData.h
@@ -0,0 +1,14 @@
+﻿// Copyright Epic Games, Inc. All Rights Reserved.
+
+#pragma once
+
+#include "AudioMotorSimConfigData.generated.h"
+
+/**
+ *  Base struct containing configuration data for motor sim components
+ */
+USTRUCT(BlueprintType)
+struct AUDIOMOTORSIM_API FAudioMotorSimConfigData
+{
+	GENERATED_BODY()
+};
diff --git a/Engine/Plugins/Runtime/AudioMotorSim/Source/AudioMotorSim/Public/IAudioMotorSim.h b/Engine/Plugins/Runtime/AudioMotorSim/Source/AudioMotorSim/Public/IAudioMotorSim.h
--- a/Engine/Plugins/Runtime/AudioMotorSim/Source/AudioMotorSim/Public/IAudioMotorSim.h
+++ b/Engine/Plugins/Runtime/AudioMotorSim/Source/AudioMotorSim/Public/IAudioMotorSim.h
@@ -3,8 +3,11 @@
 
 #include "AudioMotorSimTypes.h"
 #include "Components/ActorComponent.h"
+#include "StructUtils/InstancedStruct.h"
+
 #include "IAudioMotorSim.generated.h"
 
+struct FAudioMotorSimConfigData;
 struct FAudioMotorSimInputContext;
 struct FAudioMotorSimRuntimeContext;
 
@@ -27,6 +30,9 @@ public:
 
 	UFUNCTION(BlueprintCallable, Category = "AudioMotorSim")
 	virtual bool GetEnabled() { return false; }
+
+	UFUNCTION(BlueprintCallable, Category = "AudioMotorSim")
+	virtual void ConfigMotorSim(const FInstancedStruct& InConfigData) {};
 };
 
 UCLASS(Abstract, Blueprintable, EditInlineNew, Category = "AudioMotorSim", hideCategories = (Activation, AssetUserData, ComponentReplication, Cooking, Navigation, Replication, Tags, Variable), meta=(BlueprintSpawnableComponent))
@@ -44,7 +50,7 @@ public:
 	virtual void Reset() override;
 
 	virtual bool GetEnabled() override { return bEnabled; }
-
+	
 	/* Called every tick that this component is being updated. Use "Set Members in Struct" to update values for future components in the chain. The return value does nothing.
 	* @param Input			Holds values which are not saved between update frames which represent input to the simulation
 	* @param RuntimeInfo	Holds values which are saved between update frames to represent the output or state of the simulation
@@ -61,6 +67,8 @@ public:
 	UFUNCTION(BlueprintCallable, Category="AudioMotorSim")
 	void SetEnabled(bool bNewEnabled);
 
+	virtual void ConfigMotorSim(const FInstancedStruct& InConfigData) override;
+
 #if WITH_EDITORONLY_DATA
 	// Input data after running this component
 	UPROPERTY(BlueprintReadOnly, Category = "DebugInfo")
diff --git a/Engine/Plugins/Runtime/AudioMotorSim/Source/AudioMotorSimStandardComponents/Private/MotorPhysicsSimComponent.cpp b/Engine/Plugins/Runtime/AudioMotorSim/Source/AudioMotorSimStandardComponents/Private/MotorPhysicsSimComponent.cpp
--- a/Engine/Plugins/Runtime/AudioMotorSim/Source/AudioMotorSimStandardComponents/Private/MotorPhysicsSimComponent.cpp
+++ b/Engine/Plugins/Runtime/AudioMotorSim/Source/AudioMotorSimStandardComponents/Private/MotorPhysicsSimComponent.cpp
@@ -3,8 +3,10 @@
 #include "MotorPhysicsSimComponent.h"
 #include "AudioMotorSimLogs.h"
 #include "AudioMotorSimTypes.h"
+#include "StandardAudioMotorSimConfigDataTypes.h"
 
 #include UE_INLINE_GENERATED_CPP_BY_NAME(MotorPhysicsSimComponent)
+
 void UMotorPhysicsSimComponent::Update(FAudioMotorSimInputContext& Input, FAudioMotorSimRuntimeContext& RuntimeInfo)
 {
 	if(ensure(Weight > 0.f && EngineGearRatio > 0.f) == false)
@@ -74,6 +76,18 @@ void UMotorPhysicsSimComponent::Update(FAudioMotorSimInputContext& Input, FAudio
 	Super::Update(Input, RuntimeInfo);
 }
 
+void UMotorPhysicsSimComponent::ConfigMotorSim(const FInstancedStruct& InConfigData)
+{
+	Super::ConfigMotorSim(InConfigData);
+
+	if(const FMotorPhysicsSimConfigData* MotorPhysicsSimConfigData = InConfigData.GetPtr<FMotorPhysicsSimConfigData>())
+	{
+		Weight = MotorPhysicsSimConfigData->Weight;
+		EngineTorque = MotorPhysicsSimConfigData->EngineTorque;
+		BrakingHorsePower = MotorPhysicsSimConfigData->BrakingHorsePower;
+	}
+}
+
 float UMotorPhysicsSimComponent::CalcRpm(const float InGearRatio, const float InSpeed) const
 {
 	return InSpeed * InGearRatio / EngineGearRatio;
diff --git a/Engine/Plugins/Runtime/AudioMotorSim/Source/AudioMotorSimStandardComponents/Public/MotorPhysicsSimComponent.h b/Engine/Plugins/Runtime/AudioMotorSim/Source/AudioMotorSimStandardComponents/Public/MotorPhysicsSimComponent.h
--- a/Engine/Plugins/Runtime/AudioMotorSim/Source/AudioMotorSimStandardComponents/Public/MotorPhysicsSimComponent.h
+++ b/Engine/Plugins/Runtime/AudioMotorSim/Source/AudioMotorSimStandardComponents/Public/MotorPhysicsSimComponent.h
@@ -2,6 +2,8 @@
 #pragma once
 
 #include "IAudioMotorSim.h"
+#include "StructUtils/InstancedStruct.h"
+
 #include "MotorPhysicsSimComponent.generated.h"
 
 DECLARE_DYNAMIC_MULTICAST_DELEGATE_OneParam(FOnGearChanged, int32, NewGear);
@@ -76,6 +78,8 @@ public:
 
 	virtual void Update(FAudioMotorSimInputContext& Input, FAudioMotorSimRuntimeContext& RuntimeInfo) override;
 
+	virtual void ConfigMotorSim(const FInstancedStruct& InConfigData) override;
+
 private:
 	float CalcRpm(const float InGearRatio, const float InSpeed) const;
 
diff --git a/Engine/Plugins/Runtime/AudioMotorSim/Source/AudioMotorSimStandardComponents/Public/StandardAudioMotorSimConfigDataTypes.h b/Engine/Plugins/Runtime/AudioMotorSim/Source/AudioMotorSimStandardComponents/Public/StandardAudioMotorSimConfigDataTypes.h
new file mode 100644
--- /dev/null
+++ b/Engine/Plugins/Runtime/AudioMotorSim/Source/AudioMotorSimStandardComponents/Public/StandardAudioMotorSimConfigDataTypes.h
@@ -0,0 +1,23 @@
+﻿// Copyright Epic Games, Inc. All Rights Reserved.
+
+#pragma once
+
+#include "AudioMotorSimConfigData.h"
+
+#include "StandardAudioMotorSimConfigDataTypes.generated.h"
+
+
+USTRUCT(BlueprintType)
+struct AUDIOMOTORSIMSTANDARDCOMPONENTS_API FMotorPhysicsSimConfigData : public FAudioMotorSimConfigData
+{
+	GENERATED_BODY()
+
+	UPROPERTY(EditAnywhere, BlueprintReadOnly, Category = "Force")
+	float Weight = 900.f;
+
+	UPROPERTY(EditAnywhere, BlueprintReadOnly, Category = "Force")
+	float EngineTorque = 2500.f;
+
+	UPROPERTY(EditAnywhere, BlueprintReadOnly, Category = "Force")
+	float BrakingHorsePower = 6000.f;
+};
