Fix for cis issue  877034: Compile errors in Module.AudioMotorSimDebug.cpp and AudioMotorModelSlateIMWindow.cpp

diff --git a/Engine/Plugins/Runtime/AudioMotorSim/Source/AudioMotorSimDebug/Private/AudioMotorModelSlateIMWindow.cpp b/Engine/Plugins/Runtime/AudioMotorSim/Source/AudioMotorSimDebug/Private/AudioMotorModelSlateIMWindow.cpp
--- a/Engine/Plugins/Runtime/AudioMotorSim/Source/AudioMotorSimDebug/Private/AudioMotorModelSlateIMWindow.cpp
+++ b/Engine/Plugins/Runtime/AudioMotorSim/Source/AudioMotorSimDebug/Private/AudioMotorModelSlateIMWindow.cpp
@@ -3,16 +3,13 @@
 #include "AudioMotorModelSlateIMWindow.h"
 
 #include "AudioMotorModelComponent.h"
+#include "Engine/Engine.h"
+#include "Engine/Texture2D.h"
 #include "GameFramework/Actor.h"
+#include "SlateCore.h"
 #include "SlateIM.h"
 
-#if WITH_EDITOR
-#include "SLevelViewport.h"
-#endif
-#if WITH_ENGINE
-#include "Engine/Engine.h"
-#include "Engine/Texture2D.h"
-#endif
+
 
 namespace AudioMotorModelDebugWidgetPrivate
 {
