Fix for cis warnings

#rnx

diff --git a/Engine/Plugins/Runtime/AudioMotorSim/Source/AudioMotorSim/Private/IAudioMotorModelDebugger.cpp b/Engine/Plugins/Runtime/AudioMotorSim/Source/AudioMotorSim/Private/IAudioMotorModelDebugger.cpp
--- a/Engine/Plugins/Runtime/AudioMotorSim/Source/AudioMotorSim/Private/IAudioMotorModelDebugger.cpp
+++ b/Engine/Plugins/Runtime/AudioMotorSim/Source/AudioMotorSim/Private/IAudioMotorModelDebugger.cpp
@@ -1,3 +1,5 @@
-﻿#include "IAudioMotorModelDebugger.h"
+﻿// Copyright Epic Games, Inc. All Rights Reserved.
+
+#include "IAudioMotorModelDebugger.h"
 
 const FLazyName AudioMotorModelDebugger::DebuggerModularFeatureName = TEXT("AudioMotorModelDebugger");
diff --git a/Engine/Plugins/Runtime/AudioMotorSim/Source/AudioMotorSimDebug/AudioMotorSimDebug.Build.cs b/Engine/Plugins/Runtime/AudioMotorSim/Source/AudioMotorSimDebug/AudioMotorSimDebug.Build.cs
--- a/Engine/Plugins/Runtime/AudioMotorSim/Source/AudioMotorSimDebug/AudioMotorSimDebug.Build.cs
+++ b/Engine/Plugins/Runtime/AudioMotorSim/Source/AudioMotorSimDebug/AudioMotorSimDebug.Build.cs
@@ -1,4 +1,6 @@
-﻿using UnrealBuildTool;
+﻿// Copyright Epic Games, Inc. All Rights Reserved.
+
+using UnrealBuildTool;
 
 public class AudioMotorSimDebug : ModuleRules
 {
diff --git a/Engine/Plugins/Runtime/AudioMotorSim/Source/AudioMotorSimDebug/Private/AudioMotorModelSlateIMWindow.cpp b/Engine/Plugins/Runtime/AudioMotorSim/Source/AudioMotorSimDebug/Private/AudioMotorModelSlateIMWindow.cpp
--- a/Engine/Plugins/Runtime/AudioMotorSim/Source/AudioMotorSimDebug/Private/AudioMotorModelSlateIMWindow.cpp
+++ b/Engine/Plugins/Runtime/AudioMotorSim/Source/AudioMotorSimDebug/Private/AudioMotorModelSlateIMWindow.cpp
@@ -6,11 +6,9 @@
 #include "Engine/Engine.h"
 #include "Engine/Texture2D.h"
 #include "GameFramework/Actor.h"
-#include "SlateCore.h"
+#include "Styling/AppStyle.h"
 #include "SlateIM.h"
 
-
-
 namespace AudioMotorModelDebugWidgetPrivate
 {
 	FVector2f WindowSize(700, 900);
