Make AudioMotorModelDebugger tools type agnostic
#rnx

diff --git a/Engine/Plugins/Runtime/AudioMotorSim/Source/AudioMotorSimDebug/Private/AudioMotorModelDebuggerParamGraph.cpp b/Engine/Plugins/Runtime/AudioMotorSim/Source/AudioMotorSimDebug/Private/AudioMotorModelDebuggerParamGraph.cpp
new file mode 100644
--- /dev/null
+++ b/Engine/Plugins/Runtime/AudioMotorSim/Source/AudioMotorSimDebug/Private/AudioMotorModelDebuggerParamGraph.cpp
@@ -0,0 +1,98 @@
+// Copyright Epic Games, Inc. All Rights Reserved.
+
+#include "AudioMotorModelDebuggerParamGraph.h"
+
+AudioMotorModelDebug::FDoubleRingBufferParamGraph::FDoubleRingBufferParamGraph(const FName& ParameterName, const FParamGraphSettings& InGraphSettings)
+	: ParamName(ParameterName)
+	, GraphRange(InGraphSettings.GraphRange)
+	, GraphLineColor(InGraphSettings.GraphLineColor)
+	, bIsAdaptive(InGraphSettings.bAdaptive)
+{}
+
+void AudioMotorModelDebug::FDoubleRingBufferParamGraph::DrawGraph(const double& ParamValue)
+{
+	if (GraphValues.Num() >= NumFrames)
+	{
+		GraphValues.PopFront();
+	}
+			
+	GraphValues.Emplace(ParamValue);
+
+	SlateIM::Padding(FMargin(1, 1));
+
+	SlateIM::BeginVerticalStack();
+	{
+		SlateIM::BeginHorizontalStack();
+		{
+			SlateIM::Text(ParamName.ToString(), GraphLineColor);
+			SlateIM::Text(FString::Printf(TEXT("%0.3lf"), ParamValue), GraphLineColor);
+		}
+		SlateIM::EndHorizontalStack();
+
+
+		SlateIM::BeginHorizontalStack();
+
+		SlateIM::MaxHeight(100.f * Scale);
+		SlateIM::MaxWidth(200.f * Scale);
+
+				
+		SlateIM::BeginGraph();
+
+		if(ParamValue > GraphRange.GetUpperBoundValue())
+		{
+			GraphRange.SetUpperBoundValue(ParamValue);
+		}
+				
+		if(ParamValue < GraphRange.GetLowerBoundValue())
+		{
+			GraphRange.SetLowerBoundValue(ParamValue);
+		}
+				
+		SlateIM::GraphLine(GraphValues.Compact(), GraphLineColor, 3.f, GraphRange);
+		SlateIM::EndGraph();
+				
+		SlateIM::BeginVerticalStack();
+		SlateIM::Padding(FMargin(0, 0));
+		SlateIM::VAlign(VAlign_Top);
+		SlateIM::Text(FString::Printf(TEXT("%0.3lf"), GraphRange.GetUpperBoundValue()), GraphLineColor);
+		SlateIM::Fill();
+		SlateIM::VAlign(VAlign_Bottom);
+		SlateIM::Text(FString::Printf(TEXT("%0.3lf"), GraphRange.GetLowerBoundValue()), GraphLineColor);
+		SlateIM::EndVerticalStack();
+
+		SlateIM::EndHorizontalStack();
+
+	}
+	SlateIM::EndVerticalStack();
+}
+
+AudioMotorModelDebug::FPropertyGraph::FPropertyGraph(const FName& InPropertyName, const FProperty& InParamProperty, const void* PropertyContainerPtr, const FParamGraphSettings& InGraphSettings)
+	: FDoubleRingBufferParamGraph(InPropertyName, InGraphSettings)
+	, PropertyValuePtr(InParamProperty.ContainerPtrToValuePtr<void>(PropertyContainerPtr))
+{
+	if (const FIntProperty* IntProp = CastField<FIntProperty>(&InParamProperty))
+	{
+		PropertyFieldVariant.Set<const FIntProperty*>(IntProp);
+	}
+	else if (const FFloatProperty* FloatProp = CastField<FFloatProperty>(&InParamProperty))
+	{
+		PropertyFieldVariant.Set<const FFloatProperty*>(FloatProp);
+	}
+	else if (const FBoolProperty* BoolProp = CastField<FBoolProperty>(&InParamProperty))
+	{
+		PropertyFieldVariant.Set<const FBoolProperty*>(BoolProp);
+	}
+	else
+	{
+		ensureMsgf(false, TEXT("Unsupported property type for graphing"));
+	}
+}
+
+void AudioMotorModelDebug::FPropertyGraph::Draw()
+{
+	const double PropertyValue = GetPropertyValue<double>();
+	DrawGraph(PropertyValue);
+	
+	GraphValues.Emplace(PropertyValue);
+}
+
diff --git a/Engine/Plugins/Runtime/AudioMotorSim/Source/AudioMotorSimDebug/Private/AudioMotorModelDebuggerParamGraph.h b/Engine/Plugins/Runtime/AudioMotorSim/Source/AudioMotorSimDebug/Private/AudioMotorModelDebuggerParamGraph.h
--- a/Engine/Plugins/Runtime/AudioMotorSim/Source/AudioMotorSimDebug/Private/AudioMotorModelDebuggerParamGraph.h
+++ b/Engine/Plugins/Runtime/AudioMotorSim/Source/AudioMotorSimDebug/Private/AudioMotorModelDebuggerParamGraph.h
@@ -4,10 +4,18 @@
 
 #include "Containers/RingBuffer.h"
 #include "IAudioMotorModelDebugger.h"
+#include "Misc/TVariant.h"
 #include "SlateIM.h"
 
 namespace AudioMotorModelDebug
 {
+	struct FParamGraphSettings
+	{
+		FDoubleRange GraphRange = FDoubleRange(0, 1);
+		FLinearColor GraphLineColor = FLinearColor::MakeRandomColor();
+		bool bAdaptive = true;
+	};
+	
 	class IParamGraph
 	{
 	public:
@@ -19,126 +27,124 @@ namespace AudioMotorModelDebug
 		virtual void SetNumFrames(const uint32 InNumFrames) = 0;
 	};
 
-	struct FParamGraphSettings
-	{
-		FDoubleRange GraphRange = FDoubleRange(0, 1);
-		FLinearColor GraphLineColor = FLinearColor::MakeRandomColor();
-		bool bAdaptive = true;
-	};
-	
-	template <typename ParamType>
-	class FParamGraph : public IParamGraph
+	template <typename RingBufferType>
+	class TParamGraphRingBuffer : public IParamGraph
 	{
 	public:
+		virtual void Draw() override {};
+		virtual void SetScale(const float InScale) override
+		{
+			Scale = InScale;
+		}
 		
-		FParamGraph(const FName& InParamName, const ParamType* InParamToDraw, const FParamGraphSettings& InGraphSettings = FParamGraphSettings())
-			: ParamName(InParamName)
-			, ParamToDraw(InParamToDraw)
-			, GraphRange(InGraphSettings.GraphRange)
-			, GraphLineColor(InGraphSettings.GraphLineColor)
-			, bIsAdaptive(InGraphSettings.bAdaptive)
-		{}
-
-		virtual void Draw() override
+		virtual void SetNumFrames(const uint32 InNumFrames) override
 		{
-			if(ParamToDraw == nullptr)
-			{
-				return;
-			}
-	
-			if (GraphValues.Num() >= NumFrames)
+			int32 ClampedInNumFrames = FMath::Clamp(InNumFrames, 0, NumFramesMax);
+			
+			if(ClampedInNumFrames < NumFrames)
 			{
-				GraphValues.PopFront();
+				int32 FramesToPop = FMath::Abs(NumFrames - ClampedInNumFrames);
+				GraphValues.Pop(FramesToPop);	
 			}
 
-			const double ParamValue = *ParamToDraw;
+			NumFrames = InNumFrames;
+		};
 	
-			GraphValues.Emplace(ParamValue);
-
-			SlateIM::Padding(FMargin(1, 1));
+	protected:
+		float Scale = 1.f;
+		int32 NumFrames = 100;
+		uint32 NumFramesMax = 1000;
+		
+		TRingBuffer<RingBufferType> GraphValues;
+		
+	};
+	
+	class FDoubleRingBufferParamGraph : public TParamGraphRingBuffer<double>
+	{
+	public:
+		FDoubleRingBufferParamGraph(const FName& ParameterName, const FParamGraphSettings& InGraphSettings = FParamGraphSettings());
 
-			SlateIM::BeginVerticalStack();
-			{
-				SlateIM::BeginHorizontalStack();
-				{
-					SlateIM::Text(ParamName.ToString(), GraphLineColor);
-					SlateIM::Text(FString::Printf(TEXT("%0.3lf"), ParamValue), GraphLineColor);
-				}
-				SlateIM::EndHorizontalStack();
+	protected:
+		void DrawGraph(const double& ParamValue);
 
+		const FName ParamName;
 
-				SlateIM::BeginHorizontalStack();
+		FDoubleRange GraphRange;
+		FLinearColor GraphLineColor;
+		
+		bool bIsAdaptive = false;
+	};
 
-				SlateIM::MaxHeight(100.f * Scale);
-				SlateIM::MaxWidth(200.f * Scale);
+	
+	//Param Graph that gets its value from a FProperty in the provided container ptr, WHICH IS NOT kept alive by the graph
+	class FPropertyGraph : public FDoubleRingBufferParamGraph
+	{
+	public:
+		FPropertyGraph(const FName& InPropertyName, const FProperty& InParamProperty, const void* PropertyContainerPtr, const FParamGraphSettings& InGraphSettings = FParamGraphSettings());
 
-				
-				SlateIM::BeginGraph();
+		virtual void Draw() override;
 
-				if(ParamValue > GraphRange.GetUpperBoundValue())
+	private:
+		template <typename OutValue>
+		OutValue GetPropertyValue()
+		{
+			if(!ensureMsgf(PropertyValuePtr, TEXT("Null property value ptr, property will not be graphed")))
+			{
+				return OutValue();
+			}
+			
+			if (const FIntProperty** StoredIntPropPtr = PropertyFieldVariant.TryGet<const FIntProperty*>())
+			{
+				if(const FIntProperty* IntPropPtr = *StoredIntPropPtr)
 				{
-					GraphRange.SetUpperBoundValue(ParamValue);
+					return static_cast<OutValue>(IntPropPtr->GetPropertyValue(PropertyValuePtr));
 				}
-				
-				if(ParamValue < GraphRange.GetLowerBoundValue())
+			}
+			else if (const FFloatProperty** StoredFloatPropPtr = PropertyFieldVariant.TryGet<const FFloatProperty*>())
+			{
+				if(const FFloatProperty* FloatPropPtr = *StoredFloatPropPtr)
 				{
-					GraphRange.SetLowerBoundValue(ParamValue);
+					return static_cast<OutValue>(FloatPropPtr->GetPropertyValue(PropertyValuePtr));
 				}
-				
-				SlateIM::GraphLine(GraphValues.Compact(), GraphLineColor, 3.f, GraphRange);
-				SlateIM::EndGraph();
-				
-				SlateIM::BeginVerticalStack();
-				SlateIM::Padding(FMargin(0, 0));
-				SlateIM::VAlign(VAlign_Top);
-				SlateIM::Text(FString::Printf(TEXT("%0.3lf"), GraphRange.GetUpperBoundValue()), GraphLineColor);
-				SlateIM::Fill();
-				SlateIM::VAlign(VAlign_Bottom);
-				SlateIM::Text(FString::Printf(TEXT("%0.3lf"), GraphRange.GetLowerBoundValue()), GraphLineColor);
-				SlateIM::EndVerticalStack();
-
-				
-				SlateIM::EndHorizontalStack();
-
-				
 			}
-			SlateIM::EndVerticalStack();
-
+			else if (const FBoolProperty** StoredBoolPropPtr = PropertyFieldVariant.TryGet<const FBoolProperty*>())
+			{
+				if(const FBoolProperty* BoolPropPtr = *StoredBoolPropPtr)
+				{
+					return static_cast<OutValue>(BoolPropPtr->GetPropertyValue(PropertyValuePtr));
+				}
+			}
+			
+			return OutValue();
 		}
+		
+		const void* PropertyValuePtr;
+		TVariant<FEmptyVariantState, const FIntProperty*, const FFloatProperty*, const FBoolProperty*> PropertyFieldVariant;
+	};
 
-		virtual void SetScale(const float InScale) override
-		{
-			Scale = InScale;
-		}
+	//Param Graph that gets its value from a ParamType*, WITHOUT keeping it alive
+	template <typename ParamType>
+	class FParamGraph : FDoubleRingBufferParamGraph
+	{
+	public:
+		
+		FParamGraph(const FName& InParamName, const ParamType* InParamToDraw, const FParamGraphSettings& InGraphSettings = FParamGraphSettings())
+			: FDoubleRingBufferParamGraph(InParamName, InGraphSettings)
+			, ParamToDraw(InParamToDraw)
+		{}
 
-		virtual void SetNumFrames(const uint32 InNumFrames) override
+		virtual void Draw() override
 		{
-			int32 ClampedInNumFrames = FMath::Clamp(InNumFrames, 0, NumFramesMax);
-			
-			if(ClampedInNumFrames < NumFrames)
+			if(ParamToDraw == nullptr)
 			{
-				int32 FramesToPop = FMath::Abs(NumFrames - ClampedInNumFrames);
-				GraphValues.Pop(FramesToPop);	
+				return;
 			}
-			
-			NumFrames = ClampedInNumFrames;
-		}
 
+			DrawGraph(*ParamToDraw);
+		}
+			
 	private:
-		const FName ParamName;
-		const ParamType* ParamToDraw;
-
-		FDoubleRange GraphRange;
-		FLinearColor GraphLineColor;
-
-		bool bIsAdaptive = false;
-		
-		TRingBuffer<double> GraphValues;
-
-		float Scale = 1.f;
-		int32 NumFrames = 100;
-		uint32 NumFramesMax = 1000;
-		
+		ParamType* ParamToDraw;
 	};
 }
 
diff --git a/Engine/Plugins/Runtime/AudioMotorSim/Source/AudioMotorSimDebug/Private/AudioMotorModelSlateIMWindow.cpp b/Engine/Plugins/Runtime/AudioMotorSim/Source/AudioMotorSimDebug/Private/AudioMotorModelSlateIMWindow.cpp
--- a/Engine/Plugins/Runtime/AudioMotorSim/Source/AudioMotorSimDebug/Private/AudioMotorModelSlateIMWindow.cpp
+++ b/Engine/Plugins/Runtime/AudioMotorSim/Source/AudioMotorSimDebug/Private/AudioMotorModelSlateIMWindow.cpp
@@ -6,6 +6,7 @@
 #include "Engine/Engine.h"
 #include "Engine/Texture2D.h"
 #include "GameFramework/Actor.h"
+#include "IAudioMotorSim.h"
 #include "Styling/AppStyle.h"
 #include "SlateIM.h"
 
@@ -13,48 +14,61 @@ namespace AudioMotorModelDebugWidgetPrivate
 {
 	FVector2f WindowSize(700, 900);
 	
-	template <typename ParamType>
-	FString GenerateParamString(const FName ParamName, const ParamType Value)
+	FString GenerateParamString(const FProperty& ParamProperty, void* ParamValuePtr)
 	{
-		static_assert(
-		std::is_same_v<ParamType, float> ||
-		std::is_same_v<ParamType, int32> ||
-		std::is_same_v<ParamType, bool>,
-		"Unsupported parameter type.");
+		check(ParamValuePtr);
 
-		const FString ParamNameAsString = ParamName.ToString();
 		FString Message;
-				
-		// NOTE: This currently only supports the parameter types currently being used. 
-		// New cases will need to be added for new parameter types.
-		if constexpr (std::is_same_v<ParamType, float>)
+		FString ParamName = ParamProperty.GetName();
+		
+		if (const FIntProperty* IntProp = CastField<FIntProperty>(&ParamProperty))
 		{
-			Message = FString::Printf(TEXT("%s: %f"), *ParamNameAsString, Value);
+			const int32 Value = IntProp->GetPropertyValue(ParamValuePtr);
+			Message = FString::Printf(TEXT("%s: %d"), *ParamName, Value);
 		}
-		else if constexpr (std::is_same_v<ParamType, int32>)
+		else if (const FFloatProperty* FloatProp = CastField<FFloatProperty>(&ParamProperty))
 		{
-			Message =  FString::Printf(TEXT("%s: %d"), *ParamNameAsString, Value);
+			const float Value = FloatProp->GetPropertyValue(ParamValuePtr);
+			Message = FString::Printf(TEXT("%s: %f"), *ParamName, Value);
 		}
-		else if constexpr (std::is_same_v<ParamType, bool>)
+		else if (const FBoolProperty* BoolProp = CastField<FBoolProperty>(&ParamProperty))
 		{
-			Message =  FString::Printf(TEXT("%s: %s"), *ParamNameAsString, Value ? TEXT("True") : TEXT("False"));
+			const bool Value = BoolProp->GetPropertyValue(ParamValuePtr);
+			Message = FString::Printf(TEXT("%s: %d"), *ParamName, Value);
 		}
 		else
 		{
-			Message = FString::Printf(TEXT("Could not create param string for param %s"), *ParamNameAsString);
+			Message = FString::Printf(TEXT("Could not create param string for param %s"), *ParamName);
 		}
 
 		return Message;
 	}
 }
 
+
 void FAudioMotorModelDebugWidget::Draw()
 {
+	UpdateDebugContexts();
+	PollParameterGraphs();
+	
 	SlateIM::BeginBorder(FAppStyle::GetBrush("ToolPanel.GroupBorder"));
 	SlateIM::Fill();
 	SlateIM::HAlign(HAlign_Fill);
 	SlateIM::VAlign(VAlign_Fill);
 	SlateIM::Padding({ 5.f });
+
+	SlateIM::BeginScrollBox();
+	{
+		DrawMotorModelSelection();
+		DrawParameterGraphs();
+	}
+	SlateIM::EndScrollBox();
+	
+	SlateIM::EndBorder();
+}
+
+void FAudioMotorModelDebugWidget::DrawMotorModelSelection()
+{
 	SlateIM::BeginScrollBox();
 	{
 		SlateIM::BeginVerticalStack();
@@ -80,6 +94,7 @@ void FAudioMotorModelDebugWidget::Draw()
 					
 							if(CheckState && SelectedMotorModel != i)
 							{
+								MotorSimDebugContexts.Empty();
 								SelectedMotorModel = i;
 							}
 					
@@ -101,62 +116,141 @@ void FAudioMotorModelDebugWidget::Draw()
 
 	if(RegisteredMotorModels.Num() > 0)
 	{
+		//Number of vehicles could change because of gameplay, so we need to check our index is still valid
+		if(SelectedMotorModel < 0 || SelectedMotorModel > RegisteredMotorModels.Num() - 1)
+		{
+			MotorSimDebugContexts.Empty();
+		}
+		
 		SelectedMotorModel = FMath::Clamp(SelectedMotorModel, 0, RegisteredMotorModels.Num() - 1);
+		
 		InstantiateMotorModelDetails(RegisteredMotorModels[SelectedMotorModel]);
 	}
 	SlateIM::EndScrollBox();
-	SlateIM::EndBorder();
-
 }
 
 void FAudioMotorModelDebugWidget::InstantiateMotorModelDetails(const UAudioMotorModelComponent* Model)
 {
 	ensure(Model);
-	DrawMotorModelParameters(Model);
+	AudioMotorModelDebug::FDebugContext* MotorModelDebugContext = CreateMotorModelObjectDebugContext(Model);
+	
+	DrawParameterSelection(*MotorModelDebugContext);
 }
 
-void FAudioMotorModelDebugWidget::DrawMotorModelParameters(const UAudioMotorModelComponent* Model)
+void FAudioMotorModelDebugWidget::RegisterTrackedObject(const AudioMotorModelDebug::FDebugContext& InTrackedObject)
 {
-	CachedRuntimeInfo = Model->GetRuntimeInfo();
-	CachedInputContext = Model->GetCachedInputData();
+	MotorSimDebugContexts.AddUnique(InTrackedObject);
+}
+
+void FAudioMotorModelDebugWidget::UnregisterTrackedObject(TObjectPtr<UObject> DebuggedObject)
+{
+	if(!ensure(DebuggedObject))
+	{
+		return;
+	}
+	
+	MotorSimDebugContexts.RemoveAllSwap([&](const AudioMotorModelDebug::FDebugContext& FDebugContext)
+	{
+		return FDebugContext.DebuggedObject == DebuggedObject;
+	});
+}
+
+void FAudioMotorModelDebugWidget::GenerateParamSelector(const FProperty& ParamProperty, void* PropertyContainer)
+{
+	ensure(PropertyContainer);
+	
+	const bool bIsParamDrawn = PropertiesToGraph.Find(&ParamProperty) != INDEX_NONE;
+	bool CheckState = bIsParamDrawn;
 	
 	SlateIM::BeginHorizontalStack();
 	{
-		SlateIM::BeginScrollBox();
-		{
-			SlateIM::Text(TEXT("Runtime Context"), FColorList::Green);
+		SlateIM::CheckBox(TEXT(""), CheckState);
 
-			GenerateParamSelector("RPM", CachedRuntimeInfo.Rpm);
-			GenerateParamSelector("Gear", CachedRuntimeInfo.Gear);
-			GenerateParamSelector("Volume", CachedRuntimeInfo.Volume);
-			GenerateParamSelector("Is Shifting", CachedRuntimeInfo.bShifting);
-			GenerateParamSelector("Pitch", CachedRuntimeInfo.Pitch);
+		if(CheckState != bIsParamDrawn)
+		{
+			if(!bIsParamDrawn)
+			{
+				PropertiesToGraph.AddUnique(&ParamProperty);
+			}
+			else
+			{
+				PropertiesToGraph.RemoveAllSwap([PropToRemoveName = ParamProperty.GetFName()](const FProperty* ArrayProperty)
+				{
+					return ArrayProperty->GetFName() == PropToRemoveName;
+				});
+			}
 		}
-		SlateIM::EndScrollBox();
 
-		SlateIM::Spacer({ 5, 1 });
+		void* ValuePtr = ParamProperty.ContainerPtrToValuePtr<void>(PropertyContainer);
+		
+		SlateIM::Text(AudioMotorModelDebugWidgetPrivate::GenerateParamString(ParamProperty, ValuePtr));
+	}
+	SlateIM::EndHorizontalStack();		
+}
 
-		SlateIM::BeginScrollBox();
-		{
-			SlateIM::Text(TEXT("Input Context"), FColorList::Green);
+TSharedPtr<AudioMotorModelDebug::IParamGraph> FAudioMotorModelDebugWidget::CreateParamGraph(const FProperty& ParamProperty, void* PropertyContainer)
+{
+	ensure(PropertyContainer);
+	return MakeShared<AudioMotorModelDebug::FPropertyGraph>(ParamProperty.GetFName(), ParamProperty, PropertyContainer);
+}
+
+void FAudioMotorModelDebugWidget::RegisterMotorModelComponent(UAudioMotorModelComponent* Component)
+{
+	if(!ensureAlways(Component))
+	{
+		return;
+	}
 	
-			GenerateParamSelector("DeltaTime", CachedInputContext.DeltaTime);
-			GenerateParamSelector("Throttle", CachedInputContext.Throttle);
-			GenerateParamSelector("Brake", CachedInputContext.Brake);
-			GenerateParamSelector("Speed", CachedInputContext.Speed);
-			GenerateParamSelector("ForwardSpeed", CachedInputContext.ForwardSpeed);
-			GenerateParamSelector("UpSpeed", CachedInputContext.UpSpeed);
-			GenerateParamSelector("SurfaceFrictionModifier", CachedInputContext.SurfaceFrictionModifier);
-			GenerateParamSelector("bDriving", CachedInputContext.bDriving);
-			GenerateParamSelector("bGrounded", CachedInputContext.bGrounded);
-			GenerateParamSelector("bCanShift", CachedInputContext.bCanShift);
-			GenerateParamSelector("bClutchEngaged", CachedInputContext.bClutchEngaged);
+	RegisteredMotorModels.AddUnique(Component);
+}
+
+void FAudioMotorModelDebugWidget::UpdateDebugContexts()
+{
+	for (int32 i = MotorSimDebugContexts.Num() - 1; i >= 0; --i)
+	{
+		AudioMotorModelDebug::FDebugContext& MotorSimDebugContext = MotorSimDebugContexts[i];
+		
+		if(!MotorSimDebugContext.DebuggedObject || !MotorSimDebugContext.DebuggedObject->IsValidLowLevel())
+		{
+			MotorSimDebugContexts.RemoveAt(i);
+			continue;
 		}
-		SlateIM::EndScrollBox();
-	}
-	SlateIM::EndHorizontalStack();
 
+		if(const UAudioMotorModelComponent* AsMotorModelComponent = Cast<UAudioMotorModelComponent>(MotorSimDebugContext.DebuggedObject))
+		{
+			for(FInstancedStruct& DebugData : MotorSimDebugContext.DebugData)
+			{
+				if(FAudioMotorSimInputContext* AsInputContext = DebugData.GetMutablePtr<FAudioMotorSimInputContext>())
+				{
+					*AsInputContext = AsMotorModelComponent->GetCachedInputData();
+				}
+				else if (FAudioMotorSimRuntimeContext* AsRuntimeContext = DebugData.GetMutablePtr<FAudioMotorSimRuntimeContext>())
+				{
+					*AsRuntimeContext = AsMotorModelComponent->GetRuntimeInfo();
+				}
+			}
+		}
+#if WITH_EDITORONLY_DATA
+		else if(const UAudioMotorSimComponent* AsMotorSimComponent = Cast<UAudioMotorSimComponent>(MotorSimDebugContext.DebuggedObject))
+		{
+			for(FInstancedStruct& DebugData : MotorSimDebugContext.DebugData)
+			{
+				if(FAudioMotorSimInputContext* AsInputContext = DebugData.GetMutablePtr<FAudioMotorSimInputContext>())
+				{
+					*AsInputContext = AsMotorSimComponent->CachedInput;
+				}
+				else if (FAudioMotorSimRuntimeContext* AsRuntimeContext = DebugData.GetMutablePtr<FAudioMotorSimRuntimeContext>())
+				{
+					*AsRuntimeContext = AsMotorSimComponent->CachedRuntimeInfo;
+				}
+			}
+		}
+#endif
+	}
+}
 
+void FAudioMotorModelDebugWidget::DrawParameterGraphs()
+{
 	SlateIM::BeginVerticalStack();
 	{
 		SlateIM::Text(TEXT("Parameter Graphs"), FColorList::Green);
@@ -168,52 +262,160 @@ void FAudioMotorModelDebugWidget::DrawMotorModelParameters(const UAudioMotorMode
 	}
 	SlateIM::EndVerticalStack();
 	
-	for(const TPair<FName, TSharedPtr<AudioMotorModelDebug::IParamGraph>>& Pair : ParamGraphs)
+	for(AudioMotorModelDebug::FDebugContext& FDebugContext : MotorSimDebugContexts)
 	{
-		if(Pair.Value)
+		for(TPair<FName, TSharedPtr<AudioMotorModelDebug::IParamGraph>>& PropertyGraphPair : FDebugContext.ParamGraphs)
 		{
-			Pair.Value->Draw();
-			Pair.Value->SetScale(ParamGraphScale);
-			Pair.Value->SetNumFrames(ParamGraphNumFrames);
+			if(PropertyGraphPair.Value)
+			{
+				PropertyGraphPair.Value->Draw();
+				PropertyGraphPair.Value->SetScale(ParamGraphScale);
+				PropertyGraphPair.Value->SetNumFrames(ParamGraphNumFrames);
+			}
 		}
 	}
 }
 
-template <typename ParamType>
-void FAudioMotorModelDebugWidget::GenerateParamSelector(const FName ParamName, const ParamType& ParamToDraw)
+void FAudioMotorModelDebugWidget::PollParameterGraphs()
 {
-	ensure(RegisteredMotorModels[SelectedMotorModel]);
-	
-	FName ParameterKey = FName(FString::Printf(TEXT("%s_%s"), *RegisteredMotorModels[SelectedMotorModel]->GetOwner()->GetName(), *ParamName.ToString()));
-	
-	const bool bIsParamDrawn = ParamGraphs.Find(ParamName) != nullptr;
-	bool CheckState = bIsParamDrawn;
+	for (auto It = PropertiesToGraph.CreateIterator(); It; ++It)
+    {
+        const FProperty* Property = *It;
+    
+		if(!ensure(Property))
+		{
+			It.RemoveCurrent();
+			continue;
+		}
+		
+		for(AudioMotorModelDebug::FDebugContext& FDebugContext : MotorSimDebugContexts)
+		{
+			if(!FDebugContext.ParamGraphs.Find(Property->GetFName()))
+			{
+				TSharedPtr<AudioMotorModelDebug::IParamGraph> ParamGraphPtr = nullptr;
 
-	SlateIM::BeginHorizontalStack();
+				for(FInstancedStruct& DebugData : FDebugContext.DebugData)
+				{
+					if(DebugData.GetScriptStruct()->FindPropertyByName(Property->GetFName()))
+					{
+						ParamGraphPtr = CreateParamGraph(*Property, DebugData.GetMutableMemory());
+					}
+				}
+
+				if(ParamGraphPtr)
+				{
+					FDebugContext.ParamGraphs.Add(Property->GetFName(), ParamGraphPtr);	
+				}
+
+			}
+		}
+	}
+
+	//Remove properties that no longer need to be graphed
+	for(AudioMotorModelDebug::FDebugContext& FDebugContext : MotorSimDebugContexts)
 	{
-		SlateIM::CheckBox(TEXT(""), CheckState);
+		if(PropertiesToGraph.IsEmpty())
+		{
+			FDebugContext.ParamGraphs.Empty();
+		}
+		else
+		{
+			for(const TPair<FName, TSharedPtr<AudioMotorModelDebug::IParamGraph>>& PropertyGraphPair : FDebugContext.ParamGraphs)
+			{
+				auto FindMatchingParamGraph = [GraphedPropertyName = PropertyGraphPair.Key](const FProperty* ArrayProperty)
+				{
+					return ArrayProperty && GraphedPropertyName == ArrayProperty->GetFName();
+				};
+			
+				TArray<FName> PropertiesToRemove;
 
-		if(CheckState != bIsParamDrawn)
+				if(PropertiesToGraph.FindByPredicate(FindMatchingParamGraph) == nullptr)
+				{
+					PropertiesToRemove.Add(PropertyGraphPair.Key);
+				}
+			
+				for(const FName& PropertyToRemove : PropertiesToRemove)
+				{
+					FDebugContext.ParamGraphs.Remove(PropertyToRemove);
+				}
+			}
+		}
+	}
+}
+
+void FAudioMotorModelDebugWidget::DrawParameterSelection(AudioMotorModelDebug::FDebugContext& FDebugContext)
+{
+	SlateIM::BeginHorizontalStack();
+	{
+		SlateIM::BeginScrollBox();
 		{
-			if(!bIsParamDrawn)
+			SlateIM::Text(TEXT("Runtime Context"), FColorList::Green);
+			
+			auto FindRuntimeContextData = [](const FInstancedStruct& StoredStruct)
+			{
+				return StoredStruct.GetPtr<FAudioMotorSimRuntimeContext>() != nullptr;
+			};
+			
+			if(FInstancedStruct* RuntimeContext = FDebugContext.DebugData.FindByPredicate(FindRuntimeContextData))
 			{
-				TSharedPtr<AudioMotorModelDebug::FParamGraph<ParamType>> ParamGraph = MakeShared<AudioMotorModelDebug::FParamGraph<ParamType>>(ParamName, &ParamToDraw);
-				ParamGraphs.Emplace(ParamName, ParamGraph);
+				for (TFieldIterator<FProperty> PropIt(FAudioMotorSimRuntimeContext::StaticStruct()); PropIt; ++PropIt)
+				{
+					const FProperty* Property = *PropIt;
+
+					if(!Property)
+					{
+						continue;
+					}
+				
+					GenerateParamSelector(*Property, RuntimeContext->GetMutableMemory());
+				}
 			}
-			else
+		}
+		SlateIM::EndScrollBox();
+
+		SlateIM::Spacer({ 5, 1 });
+
+		SlateIM::BeginScrollBox();
+		{
+			SlateIM::Text(TEXT("Input Context"), FColorList::Green);
+
+			auto FindInputContextData = [](const FInstancedStruct& StoredStruct)
 			{
-				ParamGraphs.FindAndRemoveChecked(ParamName);
+				return StoredStruct.GetPtr<FAudioMotorSimInputContext>() != nullptr;
+			};
+			
+			if(FInstancedStruct* InputContext = FDebugContext.DebugData.FindByPredicate(FindInputContextData))
+			{
+				for (TFieldIterator<FProperty> PropIt(FAudioMotorSimInputContext::StaticStruct()); PropIt; ++PropIt)
+				{
+					const FProperty* Property = *PropIt;
+
+					if(!Property)
+					{
+						continue;
+					}
+				
+					GenerateParamSelector(*Property, InputContext->GetMutableMemory());
+				}
 			}
 		}
-		
-		SlateIM::Text(AudioMotorModelDebugWidgetPrivate::GenerateParamString(ParamName, ParamToDraw));
+		SlateIM::EndScrollBox();
 	}
-	SlateIM::EndHorizontalStack();		
+	SlateIM::EndHorizontalStack();
 }
 
-void FAudioMotorModelDebugWidget::RegisterMotorModelComponent(UAudioMotorModelComponent* Component)
+AudioMotorModelDebug::FDebugContext* FAudioMotorModelDebugWidget::CreateMotorModelObjectDebugContext(const UObject* ObjectToTrack)
 {
-	RegisteredMotorModels.AddUnique(Component);
+	ensure(ObjectToTrack);
+
+	uint32 Index = MotorSimDebugContexts.AddUnique(AudioMotorModelDebug::FDebugContext(ObjectToTrack));
+	
+	ensure(&MotorSimDebugContexts[Index]);
+
+	MotorSimDebugContexts[Index].DebugData.Add(FInstancedStruct::Make<FAudioMotorSimInputContext>());
+	MotorSimDebugContexts[Index].DebugData.Add(FInstancedStruct::Make<FAudioMotorSimRuntimeContext>());
+
+	return &MotorSimDebugContexts[Index];
 }
 
 FAudioMotorModelSlateIMWindow::FAudioMotorModelSlateIMWindow(const TCHAR* Command, const TCHAR* CommandHelp)
@@ -226,6 +428,7 @@ void FAudioMotorModelSlateIMWindow::DrawWindow(float DeltaTime)
 	DebugWidget.Draw();
 }
 
+
 FAudioMotorModelSlateIMViewportWidget::FAudioMotorModelSlateIMViewportWidget(const TCHAR* Command, const TCHAR* CommandHelp)
 	: FSlateIMWidgetWithCommandBase(Command, CommandHelp)
 {                                                                                
diff --git a/Engine/Plugins/Runtime/AudioMotorSim/Source/AudioMotorSimDebug/Private/AudioMotorModelSlateIMWindow.h b/Engine/Plugins/Runtime/AudioMotorSim/Source/AudioMotorSimDebug/Private/AudioMotorModelSlateIMWindow.h
--- a/Engine/Plugins/Runtime/AudioMotorSim/Source/AudioMotorSimDebug/Private/AudioMotorModelSlateIMWindow.h
+++ b/Engine/Plugins/Runtime/AudioMotorSim/Source/AudioMotorSimDebug/Private/AudioMotorModelSlateIMWindow.h
@@ -6,24 +6,63 @@
 
 #include "AudioMotorModelDebuggerParamGraph.h"
 #include "AudioMotorSimTypes.h"
+#include "StructUtils/InstancedStruct.h"
 
 class UAudioMotorModelComponent;
 
+namespace AudioMotorModelDebug
+{
+	struct FDebugContext
+	{
+		FDebugContext() = default;
+		
+		explicit FDebugContext(const TObjectPtr<const UObject>& InDebuggedObject)
+			: DebuggedObject(InDebuggedObject)
+		{}
+
+		FDebugContext(const TObjectPtr<const UObject>& InDebuggedObject, const TArray<FInstancedStruct>& InDebugData, const TMap<FName, TSharedPtr<IParamGraph>>& InParamGraphs)
+			: DebuggedObject(InDebuggedObject)
+			, DebugData(InDebugData)
+			, ParamGraphs(InParamGraphs)
+		{}
+				
+		bool operator==(const FDebugContext& A) const
+		{
+			return DebuggedObject == A.DebuggedObject;
+		}
+		
+		TObjectPtr<const UObject> DebuggedObject;
+		TArray<FInstancedStruct> DebugData;
+		TMap<FName, TSharedPtr<IParamGraph>> ParamGraphs;
+	};
+}
+
 class FAudioMotorModelDebugWidget
 {
 public:
 	void Draw();
-	void DrawMotorModelParameters(const UAudioMotorModelComponent* Model);
-
 	void RegisterMotorModelComponent(UAudioMotorModelComponent* Component);
 
 private:
+	void UpdateDebugContexts();
+	
+	void DrawMotorModelSelection();
+	void DrawParameterSelection(AudioMotorModelDebug::FDebugContext& FDebugContext);
+	
 	void InstantiateMotorModelDetails(const UAudioMotorModelComponent* Model);
+	AudioMotorModelDebug::FDebugContext* CreateMotorModelObjectDebugContext(const UObject* ObjectToTrack);
+	
+	void GenerateParamSelector(const FProperty& ParamProperty, void* PropertyContainer);
+	void DrawParameterGraphs();
+	void PollParameterGraphs();
+	
+	TSharedPtr<AudioMotorModelDebug::IParamGraph> CreateParamGraph(const FProperty& ParamProperty, void* PropertyContainer);
+	
+	void RegisterTrackedObject(const AudioMotorModelDebug::FDebugContext& InTrackedObject);
+	void UnregisterTrackedObject(TObjectPtr<UObject> DebuggedObject);
 	
 	TArray<TObjectPtr<UAudioMotorModelComponent>> RegisteredMotorModels;
-
-	TMap<FName, TSharedPtr<AudioMotorModelDebug::IParamGraph>> ParamGraphs;
-
+	
 	int32 SelectedMotorModel = 0;
 	
 	float ParamGraphScale = 1.0f;         
@@ -31,12 +70,9 @@ private:
 	
 	float ParamGraphNumFrames = 300.0f;         
 	float ParamGraphNumFramesMax = 1000.0f;        
-
-	template <typename ParamType>
-	void GenerateParamSelector(const FName ParamName, const ParamType& ParamToDraw);
 	
-	FAudioMotorSimRuntimeContext CachedRuntimeInfo;
-	FAudioMotorSimInputContext CachedInputContext;
+	TArray<AudioMotorModelDebug::FDebugContext>	MotorSimDebugContexts;
+	TArray<const FProperty*> PropertiesToGraph;
 	
 };
 
