First pass on debug drawing motor sim components
#rnx

diff --git a/Engine/Plugins/Runtime/AudioMotorSim/Source/AudioMotorSimDebug/Private/AudioMotorModelSlateIMWindow.cpp b/Engine/Plugins/Runtime/AudioMotorSim/Source/AudioMotorSimDebug/Private/AudioMotorModelSlateIMWindow.cpp
--- a/Engine/Plugins/Runtime/AudioMotorSim/Source/AudioMotorSimDebug/Private/AudioMotorModelSlateIMWindow.cpp
+++ b/Engine/Plugins/Runtime/AudioMotorSim/Source/AudioMotorSimDebug/Private/AudioMotorModelSlateIMWindow.cpp
@@ -133,8 +133,64 @@ void FAudioMotorModelDebugWidget::InstantiateMotorModelDetails(const UAudioMotor
 {
 	ensure(Model);
 	AudioMotorModelDebug::FDebugContext* MotorModelDebugContext = CreateMotorModelObjectDebugContext(Model);
+	SlateIM::BeginHorizontalStack();
+	{
+		DrawParameterSelection(*MotorModelDebugContext);
+		DrawMotorSimComponentList(Model);
+	}
+	SlateIM::EndHorizontalStack();
+	
+}
+
+void FAudioMotorModelDebugWidget::DrawMotorSimComponentList(const UAudioMotorModelComponent* Model)
+{
+	if(!ensure(Model))
+	{
+		return;
+	}
+	
+	SlateIM::BeginScrollBox();
+	{
+		SlateIM::Text(TEXT("Motor Sim Components"), FColorList::Green);
+
+		for(const FMotorSimEntry& MotorSimComponent : Model->SimComponents)
+		{
+			if(const UObject* MotorSimObj = MotorSimComponent.Sim.GetObject())
+			{
+				SlateIM::BeginHorizontalStack();
+				{
+					auto FindMotorSimDebugContext = [MotorSimObj](const AudioMotorModelDebug::FDebugContext& DebugContext)
+					{
+							return MotorSimObj == DebugContext.DebuggedObject;
+					};
+
+					const AudioMotorModelDebug::FDebugContext* DebugContext = MotorSimDebugContexts.FindByPredicate(FindMotorSimDebugContext);
+
+					const bool bIsMotorSimDrawn = DebugContext != nullptr;
+					bool bCheckBox = bIsMotorSimDrawn;
+					SlateIM::CheckBox(TEXT(""), bCheckBox);
+					FString ComponentName = FString::Printf(TEXT("%s"), *MotorSimObj->GetName());
+					SlateIM::Text(ComponentName);
+					
+					if(bIsMotorSimDrawn != bCheckBox)
+					{
+						if(bIsMotorSimDrawn)
+						{
+							const int32 DebugContextIndex = MotorSimDebugContexts.Find(*DebugContext);
+							MotorSimDebugContexts.RemoveAtSwap(DebugContextIndex);
+						}
+						else 
+						{
+							CreateMotorModelObjectDebugContext(MotorSimObj);
+						}
+					}
+				}
+				SlateIM::EndHorizontalStack();
+			}
+		}
+	}
+	SlateIM::EndScrollBox();
 	
-	DrawParameterSelection(*MotorModelDebugContext);
 }
 
 void FAudioMotorModelDebugWidget::RegisterTrackedObject(const AudioMotorModelDebug::FDebugContext& InTrackedObject)
@@ -149,9 +205,9 @@ void FAudioMotorModelDebugWidget::UnregisterTrackedObject(TObjectPtr<UObject> De
 		return;
 	}
 	
-	MotorSimDebugContexts.RemoveAllSwap([&](const AudioMotorModelDebug::FDebugContext& FDebugContext)
+	MotorSimDebugContexts.RemoveAllSwap([&](const AudioMotorModelDebug::FDebugContext& DebugContext)
 	{
-		return FDebugContext.DebuggedObject == DebuggedObject;
+		return DebugContext.DebuggedObject == DebuggedObject;
 	});
 }
 
@@ -194,6 +250,12 @@ TSharedPtr<AudioMotorModelDebug::IParamGraph> FAudioMotorModelDebugWidget::Creat
 	return MakeShared<AudioMotorModelDebug::FPropertyGraph>(ParamProperty.GetFName(), ParamProperty, PropertyContainer);
 }
 
+TSharedPtr<AudioMotorModelDebug::IParamGraph> FAudioMotorModelDebugWidget::CreateParamGraph(const FName& GraphedPropertyName, const FProperty& ParamProperty, void* PropertyContainer)
+{
+	ensure(PropertyContainer);
+	return MakeShared<AudioMotorModelDebug::FPropertyGraph>(GraphedPropertyName, ParamProperty, PropertyContainer);
+}
+
 void FAudioMotorModelDebugWidget::RegisterMotorModelComponent(UAudioMotorModelComponent* Component)
 {
 	if(!ensureAlways(Component))
@@ -261,19 +323,37 @@ void FAudioMotorModelDebugWidget::DrawParameterGraphs()
 		SlateIM::Slider(ParamGraphNumFrames, 0, ParamGraphNumFramesMax, 1);
 	}
 	SlateIM::EndVerticalStack();
-	
-	for(AudioMotorModelDebug::FDebugContext& FDebugContext : MotorSimDebugContexts)
+
+	SlateIM::BeginHorizontalStack();
 	{
-		for(TPair<FName, TSharedPtr<AudioMotorModelDebug::IParamGraph>>& PropertyGraphPair : FDebugContext.ParamGraphs)
+		for(const FProperty* PropertyToGraph : PropertiesToGraph)
 		{
-			if(PropertyGraphPair.Value)
+			if(!PropertyToGraph)
 			{
-				PropertyGraphPair.Value->Draw();
-				PropertyGraphPair.Value->SetScale(ParamGraphScale);
-				PropertyGraphPair.Value->SetNumFrames(ParamGraphNumFrames);
+				continue;
 			}
+
+			SlateIM::BeginVerticalStack();
+			{
+				for(AudioMotorModelDebug::FDebugContext& DebugContext : MotorSimDebugContexts)
+				{
+		
+					for(TPair<FName, TSharedPtr<AudioMotorModelDebug::IParamGraph>>& PropertyGraphPair : DebugContext.ParamGraphs)
+					{
+						if(PropertyToGraph->GetName() == PropertyGraphPair.Key && PropertyGraphPair.Value)
+						{
+							PropertyGraphPair.Value->Draw();
+							PropertyGraphPair.Value->SetScale(ParamGraphScale);
+							PropertyGraphPair.Value->SetNumFrames(ParamGraphNumFrames);
+						}
+					}
+				}
+			}
+			SlateIM::EndVerticalStack();
 		}
 	}
+	SlateIM::EndHorizontalStack();
+
 }
 
 void FAudioMotorModelDebugWidget::PollParameterGraphs()
@@ -288,23 +368,24 @@ void FAudioMotorModelDebugWidget::PollParameterGraphs()
 			continue;
 		}
 		
-		for(AudioMotorModelDebug::FDebugContext& FDebugContext : MotorSimDebugContexts)
+		for(AudioMotorModelDebug::FDebugContext& DebugContext : MotorSimDebugContexts)
 		{
-			if(!FDebugContext.ParamGraphs.Find(Property->GetFName()))
+			if(!DebugContext.ParamGraphs.Find(Property->GetFName()))
 			{
 				TSharedPtr<AudioMotorModelDebug::IParamGraph> ParamGraphPtr = nullptr;
 
-				for(FInstancedStruct& DebugData : FDebugContext.DebugData)
+				for(FInstancedStruct& DebugData : DebugContext.DebugData)
 				{
 					if(DebugData.GetScriptStruct()->FindPropertyByName(Property->GetFName()))
 					{
-						ParamGraphPtr = CreateParamGraph(*Property, DebugData.GetMutableMemory());
+						const FString GraphName = FString::Printf(TEXT("%s (%s)"), *Property->GetName(), DebugContext.DebuggedObject ? *DebugContext.DebuggedObject->GetName() : TEXT("NULL"));
+						ParamGraphPtr = CreateParamGraph(FName(GraphName), *Property, DebugData.GetMutableMemory());
 					}
 				}
 
 				if(ParamGraphPtr)
 				{
-					FDebugContext.ParamGraphs.Add(Property->GetFName(), ParamGraphPtr);	
+					DebugContext.ParamGraphs.Add(Property->GetFName(), ParamGraphPtr);	
 				}
 
 			}
@@ -343,7 +424,7 @@ void FAudioMotorModelDebugWidget::PollParameterGraphs()
 	}
 }
 
-void FAudioMotorModelDebugWidget::DrawParameterSelection(AudioMotorModelDebug::FDebugContext& FDebugContext)
+void FAudioMotorModelDebugWidget::DrawParameterSelection(AudioMotorModelDebug::FDebugContext& DebugContext)
 {
 	SlateIM::BeginHorizontalStack();
 	{
@@ -356,7 +437,7 @@ void FAudioMotorModelDebugWidget::DrawParameterSelection(AudioMotorModelDebug::F
 				return StoredStruct.GetPtr<FAudioMotorSimRuntimeContext>() != nullptr;
 			};
 			
-			if(FInstancedStruct* RuntimeContext = FDebugContext.DebugData.FindByPredicate(FindRuntimeContextData))
+			if(FInstancedStruct* RuntimeContext = DebugContext.DebugData.FindByPredicate(FindRuntimeContextData))
 			{
 				for (TFieldIterator<FProperty> PropIt(FAudioMotorSimRuntimeContext::StaticStruct()); PropIt; ++PropIt)
 				{
@@ -384,7 +465,7 @@ void FAudioMotorModelDebugWidget::DrawParameterSelection(AudioMotorModelDebug::F
 				return StoredStruct.GetPtr<FAudioMotorSimInputContext>() != nullptr;
 			};
 			
-			if(FInstancedStruct* InputContext = FDebugContext.DebugData.FindByPredicate(FindInputContextData))
+			if(FInstancedStruct* InputContext = DebugContext.DebugData.FindByPredicate(FindInputContextData))
 			{
 				for (TFieldIterator<FProperty> PropIt(FAudioMotorSimInputContext::StaticStruct()); PropIt; ++PropIt)
 				{
diff --git a/Engine/Plugins/Runtime/AudioMotorSim/Source/AudioMotorSimDebug/Private/AudioMotorModelSlateIMWindow.h b/Engine/Plugins/Runtime/AudioMotorSim/Source/AudioMotorSimDebug/Private/AudioMotorModelSlateIMWindow.h
--- a/Engine/Plugins/Runtime/AudioMotorSim/Source/AudioMotorSimDebug/Private/AudioMotorModelSlateIMWindow.h
+++ b/Engine/Plugins/Runtime/AudioMotorSim/Source/AudioMotorSimDebug/Private/AudioMotorModelSlateIMWindow.h
@@ -8,6 +8,7 @@
 #include "AudioMotorSimTypes.h"
 #include "StructUtils/InstancedStruct.h"
 
+struct FInstancedStruct;
 class UAudioMotorModelComponent;
 
 namespace AudioMotorModelDebug
@@ -47,9 +48,10 @@ private:
 	void UpdateDebugContexts();
 	
 	void DrawMotorModelSelection();
-	void DrawParameterSelection(AudioMotorModelDebug::FDebugContext& FDebugContext);
-	
+	void DrawParameterSelection(AudioMotorModelDebug::FDebugContext& DebugContext);
+
 	void InstantiateMotorModelDetails(const UAudioMotorModelComponent* Model);
+	void DrawMotorSimComponentList(const UAudioMotorModelComponent* Model);
 	AudioMotorModelDebug::FDebugContext* CreateMotorModelObjectDebugContext(const UObject* ObjectToTrack);
 	
 	void GenerateParamSelector(const FProperty& ParamProperty, void* PropertyContainer);
@@ -57,6 +59,7 @@ private:
 	void PollParameterGraphs();
 	
 	TSharedPtr<AudioMotorModelDebug::IParamGraph> CreateParamGraph(const FProperty& ParamProperty, void* PropertyContainer);
+	TSharedPtr<AudioMotorModelDebug::IParamGraph> CreateParamGraph(const FName& GraphedPropertyName, const FProperty& ParamProperty, void* PropertyContainer);
 	
 	void RegisterTrackedObject(const AudioMotorModelDebug::FDebugContext& InTrackedObject);
 	void UnregisterTrackedObject(TObjectPtr<UObject> DebuggedObject);
