fix for cis issue with audio motor model debugger

#rnx

diff --git a/Engine/Plugins/Runtime/AudioMotorSim/Source/AudioMotorSimDebug/Private/AudioMotorModelDebuggerParamGraph.h b/Engine/Plugins/Runtime/AudioMotorSim/Source/AudioMotorSimDebug/Private/AudioMotorModelDebuggerParamGraph.h
--- a/Engine/Plugins/Runtime/AudioMotorSim/Source/AudioMotorSimDebug/Private/AudioMotorModelDebuggerParamGraph.h
+++ b/Engine/Plugins/Runtime/AudioMotorSim/Source/AudioMotorSimDebug/Private/AudioMotorModelDebuggerParamGraph.h
@@ -39,11 +39,11 @@ namespace AudioMotorModelDebug
 		
 		virtual void SetNumFrames(const uint32 InNumFrames) override
 		{
-			int32 ClampedInNumFrames = FMath::Clamp(InNumFrames, 0, NumFramesMax);
+			int32 ClampedInNumFrames = FMath::Clamp(InNumFrames, 1, NumFramesMax);
 			
-			if(ClampedInNumFrames < NumFrames)
+			if(ClampedInNumFrames < GraphValues.Num())
 			{
-				int32 FramesToPop = FMath::Abs(NumFrames - ClampedInNumFrames);
+				int32 FramesToPop = FMath::Abs(GraphValues.Num() - ClampedInNumFrames);
 				GraphValues.Pop(FramesToPop);	
 			}
 
diff --git a/Engine/Plugins/Runtime/AudioMotorSim/Source/AudioMotorSimDebug/Private/AudioMotorModelSlateIMWindow.cpp b/Engine/Plugins/Runtime/AudioMotorSim/Source/AudioMotorSimDebug/Private/AudioMotorModelSlateIMWindow.cpp
--- a/Engine/Plugins/Runtime/AudioMotorSim/Source/AudioMotorSimDebug/Private/AudioMotorModelSlateIMWindow.cpp
+++ b/Engine/Plugins/Runtime/AudioMotorSim/Source/AudioMotorSimDebug/Private/AudioMotorModelSlateIMWindow.cpp
@@ -320,7 +320,7 @@ void FAudioMotorModelDebugWidget::DrawParameterGraphs()
 		SlateIM::Slider(ParamGraphScale, 0, ParamGraphScaleMax, 1);
 		
 		SlateIM::Text(FString::Printf(TEXT("NumFrames: %0.f"), ParamGraphNumFrames));
-		SlateIM::Slider(ParamGraphNumFrames, 0, ParamGraphNumFramesMax, 1);
+		SlateIM::Slider(ParamGraphNumFrames, 1, ParamGraphNumFramesMax, 1);
 	}
 	SlateIM::EndVerticalStack();
 
@@ -393,33 +393,34 @@ void FAudioMotorModelDebugWidget::PollParameterGraphs()
 	}
 
 	//Remove properties that no longer need to be graphed
-	for(AudioMotorModelDebug::FDebugContext& FDebugContext : MotorSimDebugContexts)
+	for(AudioMotorModelDebug::FDebugContext& DebugContext : MotorSimDebugContexts)
 	{
 		if(PropertiesToGraph.IsEmpty())
 		{
-			FDebugContext.ParamGraphs.Empty();
+			DebugContext.ParamGraphs.Empty();
 		}
 		else
 		{
-			for(const TPair<FName, TSharedPtr<AudioMotorModelDebug::IParamGraph>>& PropertyGraphPair : FDebugContext.ParamGraphs)
+			TArray<FName> PropertiesToRemove;
+			
+			for(const TPair<FName, TSharedPtr<AudioMotorModelDebug::IParamGraph>>& PropertyGraphPair : DebugContext.ParamGraphs)
 			{
 				auto FindMatchingParamGraph = [GraphedPropertyName = PropertyGraphPair.Key](const FProperty* ArrayProperty)
 				{
 					return ArrayProperty && GraphedPropertyName == ArrayProperty->GetFName();
 				};
-			
-				TArray<FName> PropertiesToRemove;
-
+				
 				if(PropertiesToGraph.FindByPredicate(FindMatchingParamGraph) == nullptr)
 				{
-					PropertiesToRemove.Add(PropertyGraphPair.Key);
+					PropertiesToRemove.AddUnique(PropertyGraphPair.Key);
 				}
+			}
 			
-				for(const FName& PropertyToRemove : PropertiesToRemove)
-				{
-					FDebugContext.ParamGraphs.Remove(PropertyToRemove);
-				}
+			for(const FName& PropertyToRemove : PropertiesToRemove)
+			{
+				DebugContext.ParamGraphs.Remove(PropertyToRemove);
 			}
+			
 		}
 	}
 }
