[AudioMotorModelDebugger] MotorSimComponent Debug and editing
#rnx

diff --git a/Engine/Plugins/Runtime/AudioMotorSim/Source/AudioMotorSimDebug/AudioMotorSimDebug.Build.cs b/Engine/Plugins/Runtime/AudioMotorSim/Source/AudioMotorSimDebug/AudioMotorSimDebug.Build.cs
--- a/Engine/Plugins/Runtime/AudioMotorSim/Source/AudioMotorSimDebug/AudioMotorSimDebug.Build.cs
+++ b/Engine/Plugins/Runtime/AudioMotorSim/Source/AudioMotorSimDebug/AudioMotorSimDebug.Build.cs
@@ -20,5 +20,14 @@ public class AudioMotorSimDebug : ModuleRules
             }
 			
         );
+        
+        if (Target.Type == TargetType.Editor)
+        {
+	        PrivateDependencyModuleNames.AddRange(
+		        new string[] {
+			        "UnrealEd",
+		        }
+	        );
+        }
     }
 }
\ No newline at end of file
diff --git a/Engine/Plugins/Runtime/AudioMotorSim/Source/AudioMotorSimDebug/Private/AudioMotorModelDebuggerPropertyUI.cpp b/Engine/Plugins/Runtime/AudioMotorSim/Source/AudioMotorSimDebug/Private/AudioMotorModelDebuggerPropertyUI.cpp
new file mode 100644
--- /dev/null
+++ b/Engine/Plugins/Runtime/AudioMotorSim/Source/AudioMotorSimDebug/Private/AudioMotorModelDebuggerPropertyUI.cpp
@@ -0,0 +1,129 @@
+﻿// Copyright Epic Games, Inc. All Rights Reserved.
+
+#include "AudioMotorModelDebuggerPropertyUI.h"
+
+#include "SlateIM.h"
+
+namespace AudioMotorModelDebuggerPropertyUIPrivate
+{
+	constexpr float MinSpinboxValue = 0.f;
+	constexpr float MaxSpinboxValue = 20000.f;
+}
+
+void AudioMotorModelDebuggerPropertyUI::DrawPropertyUI(const FProperty& InProperty, void* ContainerPtr, bool bDisplayPropertyName)
+{
+	if(!ensure(ContainerPtr))
+	{
+		return;
+	}
+
+	if(!CanDrawProperty(InProperty))
+	{
+		return;
+	}
+
+	SlateIM::BeginHorizontalStack();
+	{
+		if(bDisplayPropertyName)
+		{
+			SlateIM::VAlign(VAlign_Center);
+			SlateIM::Text(FString::Printf(TEXT("%s"), *InProperty.GetName()));
+		}
+		
+		if (const FFloatProperty* FloatProp = CastField<FFloatProperty>(&InProperty))
+		{
+			const float PropertyValue = FloatProp->GetPropertyValue_InContainer(ContainerPtr);
+			float SpinBoxValue = PropertyValue;
+
+			float MinSpinBoxValue = AudioMotorModelDebuggerPropertyUIPrivate::MinSpinboxValue;
+			float MaxSpinBoxValue = AudioMotorModelDebuggerPropertyUIPrivate::MaxSpinboxValue;
+
+#if WITH_METADATA
+			TPair<float,float> MinMax = GetMinMaxClampFromPropertyMetadata<float>(InProperty, TPair<float, float>(MinSpinBoxValue, MaxSpinBoxValue));
+			MinSpinBoxValue = MinMax.Key;
+			MaxSpinBoxValue = MinMax.Value;
+#endif
+
+			SlateIM::VAlign(VAlign_Center);
+			if (SlateIM::SpinBox(SpinBoxValue, MinSpinBoxValue, MaxSpinBoxValue))
+			{
+				if(SpinBoxValue != PropertyValue)
+				{
+					FloatProp->SetValue_InContainer(ContainerPtr, SpinBoxValue);
+				}
+			}
+		}
+		else if (const FIntProperty* IntProp = CastField<FIntProperty>(&InProperty))
+		{
+			const int32 PropertyValue = IntProp->GetPropertyValue_InContainer(ContainerPtr);
+			int32 SpinBoxValue = PropertyValue;
+
+			int32 MinSpinBoxValue = AudioMotorModelDebuggerPropertyUIPrivate::MinSpinboxValue;
+			int32 MaxSpinBoxValue = AudioMotorModelDebuggerPropertyUIPrivate::MaxSpinboxValue;
+
+#if WITH_EDITORONLY_DATA
+			TPair<int32,int32> MinMax = GetMinMaxClampFromPropertyMetadata<int32>(InProperty, TPair<int32, int32>(MinSpinBoxValue, MaxSpinBoxValue));
+			MinSpinBoxValue = MinMax.Key;
+			MaxSpinBoxValue = MinMax.Value;
+#endif
+
+			SlateIM::VAlign(VAlign_Center);
+			if (SlateIM::SpinBox(SpinBoxValue, MinSpinBoxValue, MaxSpinBoxValue))
+			{
+				if(SpinBoxValue != PropertyValue)
+				{
+					IntProp->SetValue_InContainer(ContainerPtr, SpinBoxValue);
+				}
+			}
+		}
+		else if (const FBoolProperty* BoolProp = CastField<FBoolProperty>(&InProperty))
+		{
+			const bool PropertyValue = BoolProp->GetPropertyValue_InContainer(ContainerPtr);
+			bool CheckBoxValue = PropertyValue;
+
+			SlateIM::VAlign(VAlign_Center);
+			if (SlateIM::CheckBox(TEXT(""), CheckBoxValue))
+			{
+				if(CheckBoxValue != PropertyValue)
+				{
+					BoolProp->SetValue_InContainer(ContainerPtr, &CheckBoxValue);
+				}
+			}
+			
+		}
+	}
+	SlateIM::EndHorizontalStack();
+}
+
+void AudioMotorModelDebuggerPropertyUI::DrawArrayPropertyUI(const FArrayProperty& InArrayProperty, const void* ContainerPtr, bool bDisplayArrayName)
+{
+	if(bDisplayArrayName)
+	{
+		SlateIM::Text(InArrayProperty.GetName());
+	}
+	
+	FScriptArrayHelper ArrayHelper(&InArrayProperty, ContainerPtr);
+
+	SlateIM::BeginVerticalStack();
+	for (int32 ElementIndex = 0; ElementIndex < ArrayHelper.Num(); ++ElementIndex)
+	{
+		void* ValueAddress = ArrayHelper.GetRawPtr(ElementIndex);
+		
+		SlateIM::BeginHorizontalStack();
+		{
+			SlateIM::VAlign(VAlign_Center);
+			SlateIM::Text(FString::Printf(TEXT("%d"), ElementIndex));
+			DrawPropertyUI(*InArrayProperty.Inner, ValueAddress, false);
+		}
+		SlateIM::EndHorizontalStack();
+	}
+	SlateIM::EndVerticalStack();
+}
+
+bool AudioMotorModelDebuggerPropertyUI::CanDrawProperty(const FProperty& InProperty)
+{
+	bool bSupportedType = InProperty.IsA<FFloatProperty>() || InProperty.IsA<FIntProperty>() || InProperty.IsA<FBoolProperty>();
+	return bSupportedType;
+
+}
+
diff --git a/Engine/Plugins/Runtime/AudioMotorSim/Source/AudioMotorSimDebug/Private/AudioMotorModelDebuggerPropertyUI.h b/Engine/Plugins/Runtime/AudioMotorSim/Source/AudioMotorSimDebug/Private/AudioMotorModelDebuggerPropertyUI.h
new file mode 100644
--- /dev/null
+++ b/Engine/Plugins/Runtime/AudioMotorSim/Source/AudioMotorSimDebug/Private/AudioMotorModelDebuggerPropertyUI.h
@@ -0,0 +1,59 @@
+﻿// Copyright Epic Games, Inc. All Rights Reserved.
+
+#pragma once
+
+#include "Containers/UnrealString.h"
+#include "Templates/Tuple.h"
+#include "UObject/UnrealType.h"
+
+namespace AudioMotorModelDebuggerPropertyUI
+{
+	void DrawPropertyUI(const FProperty& InProperty, void* ContainerPtr, bool bDisplayPropertyName = true);
+	void DrawArrayPropertyUI(const FArrayProperty& InArrayProperty, const void* ContainerPtr, bool bDisplayArrayName = true);
+	bool CanDrawProperty(const FProperty& InProperty);
+
+#if WITH_METADATA
+	template<typename NumericType>
+	TPair<NumericType, NumericType> GetMinMaxClampFromPropertyMetadata(const FProperty& TargetProperty, TPair<NumericType, NumericType> MinMaxLimits = TPair<NumericType, NumericType>(TNumericLimits<NumericType>::Lowest(), TNumericLimits<NumericType>::Max()))
+	{
+		static_assert(TIsArithmetic<NumericType>::Value, "Property to clamp must have numeric type");
+
+		const FString& MetaClampMinString = TargetProperty.GetMetaData("ClampMin");
+		const FString& MetaClampMaxString = TargetProperty.GetMetaData("ClampMax");
+		const FString& MetaUIMinString = TargetProperty.GetMetaData("UIMin");
+		const FString& MetaUIMaxString = TargetProperty.GetMetaData("UIMax");
+				
+		NumericType ClampMin = MinMaxLimits.Key;
+		if (!MetaClampMinString.IsEmpty())
+		{
+			TTypeFromString<NumericType>::FromString(ClampMin, *MetaClampMinString);
+		}
+
+		NumericType ClampMax = MinMaxLimits.Value;
+		if (!MetaClampMaxString.IsEmpty())
+		{
+			TTypeFromString<NumericType>::FromString(ClampMax, *MetaClampMaxString);
+		}
+		
+		NumericType UIMin = MinMaxLimits.Key;
+		if (!MetaUIMinString.IsEmpty())
+		{
+			TTypeFromString<NumericType>::FromString(UIMin, *MetaUIMinString);
+		}
+		
+		NumericType UIMax = MinMaxLimits.Value;
+		if (!MetaUIMaxString.IsEmpty())
+		{
+			TTypeFromString<NumericType>::FromString(UIMax, *MetaUIMaxString);
+		}
+
+		const NumericType MinAllowedValue = FMath::Max(UIMin, ClampMin);
+		const NumericType MaxAllowedValue = FMath::Min(UIMax, ClampMax); 
+		
+		return TPair<NumericType, NumericType>(MinAllowedValue, MaxAllowedValue);
+	}
+#endif
+	
+}
+
+
diff --git a/Engine/Plugins/Runtime/AudioMotorSim/Source/AudioMotorSimDebug/Private/AudioMotorModelSlateIMWindow.cpp b/Engine/Plugins/Runtime/AudioMotorSim/Source/AudioMotorSimDebug/Private/AudioMotorModelSlateIMWindow.cpp
--- a/Engine/Plugins/Runtime/AudioMotorSim/Source/AudioMotorSimDebug/Private/AudioMotorModelSlateIMWindow.cpp
+++ b/Engine/Plugins/Runtime/AudioMotorSim/Source/AudioMotorSimDebug/Private/AudioMotorModelSlateIMWindow.cpp
@@ -3,16 +3,25 @@
 #include "AudioMotorModelSlateIMWindow.h"
 
 #include "AudioMotorModelComponent.h"
+#include "AudioMotorModelDebuggerPropertyUI.h"
 #include "Engine/Engine.h"
 #include "Engine/Texture2D.h"
 #include "GameFramework/Actor.h"
 #include "IAudioMotorSim.h"
 #include "Styling/AppStyle.h"
 #include "SlateIM.h"
+#include "Styling/SlateTypes.h"
+
+#if WITH_METADATA
+#include "EditorCategoryUtils.h"
+#endif
 
 namespace AudioMotorModelDebugWidgetPrivate
 {
-	FVector2f WindowSize(700, 900);
+	
+	FVector2f WindowSize(850, 900);
+	FVector2f ViewportAlignment = FVector2f(0.0f, 0);
+	FAnchors ViewportAnchors = FAnchors(0.0f, 0);
 	
 	FString GenerateParamString(const FProperty& ParamProperty, void* ParamValuePtr)
 	{
@@ -43,12 +52,39 @@ namespace AudioMotorModelDebugWidgetPrivate
 
 		return Message;
 	}
+
+#if WITH_METADATA
+	bool IsPropertyHiddenByClassCategory(const FProperty& InProperty, const UClass& ContainerClass)
+	{
+		FString PropertyCategory = InProperty.GetMetaData(TEXT("Category"));
+		const UClass* ClassToCheck = &ContainerClass;
+
+		while (ClassToCheck)
+		{
+			TArray<FString> HiddenCategoryList;
+			FEditorCategoryUtils::GetClassHideCategories(ClassToCheck, HiddenCategoryList);
+			
+			for (FString& Category : HiddenCategoryList)
+			{
+				if (Category == PropertyCategory)
+				{
+					return true;
+				}
+			}
+
+			ClassToCheck = ClassToCheck->GetSuperClass();
+		}
+
+		return false;
+	}
+#endif
 }
 
 
 void FAudioMotorModelDebugWidget::Draw()
 {
 	UpdateDebugContexts();
+	PollRegisteredModels();
 	PollParameterGraphs();
 	
 	SlateIM::BeginBorder(FAppStyle::GetBrush("ToolPanel.GroupBorder"));
@@ -83,13 +119,14 @@ void FAudioMotorModelDebugWidget::DrawMotorModelSelection()
 					{
 						if(!IsValid(Model))
 						{
-							RegisteredMotorModels.RemoveAtSwap(i);
+							continue;
 						}
 				
 						SlateIM::BeginHorizontalStack();
 						{
 							bool CheckState = SelectedMotorModel == i;
 
+							SlateIM::VAlign(VAlign_Center);
 							SlateIM::CheckBox(TEXT(""), CheckState);
 					
 							if(CheckState && SelectedMotorModel != i)
@@ -99,7 +136,9 @@ void FAudioMotorModelDebugWidget::DrawMotorModelSelection()
 							}
 					
 							SlateIM::VAlign(VAlign_Center);
-							FString MotorModelOwnerName = Model->GetOwner() ? Model->GetOwner()->GetName() : FString("");
+
+							const bool bHasValidOwner = Model->GetOwner() && Model->GetOwner()->IsValidLowLevel();
+							const FString MotorModelOwnerName = bHasValidOwner ? Model->GetOwner()->GetName() : Model->GetName();
 							const bool bIsBeingDriven = Model->GetCachedInputData().bDriving && Model->GetRuntimeInfo().Rpm > 0.f;
 							SlateIM::Text(MotorModelOwnerName, bIsBeingDriven ? FColor::Green : FColor::Silver);
 						}
@@ -129,68 +168,164 @@ void FAudioMotorModelDebugWidget::DrawMotorModelSelection()
 	SlateIM::EndScrollBox();
 }
 
-void FAudioMotorModelDebugWidget::InstantiateMotorModelDetails(const UAudioMotorModelComponent* Model)
+void FAudioMotorModelDebugWidget::InstantiateMotorModelDetails(UAudioMotorModelComponent* Model)
 {
 	ensure(Model);
 	AudioMotorModelDebug::FDebugContext* MotorModelDebugContext = CreateMotorModelObjectDebugContext(Model);
+
+	if(!MotorModelDebugContext)
+	{
+		return;
+	}
+
+	if(MotorModelDebugContext)
+	{
+		MotorModelDebugContext->bDrawGraphs = true;
+	}
+	
 	SlateIM::BeginHorizontalStack();
 	{
 		DrawParameterSelection(*MotorModelDebugContext);
-		DrawMotorSimComponentList(Model);
+		CreateModelSimComponentsDebugContexts(Model);
+		DrawMotorSimComponentList();
 	}
 	SlateIM::EndHorizontalStack();
 	
 }
 
-void FAudioMotorModelDebugWidget::DrawMotorSimComponentList(const UAudioMotorModelComponent* Model)
+void FAudioMotorModelDebugWidget::DrawObjectPropertiesTree(const TObjectPtr<UObject>& InObject)
 {
-	if(!ensure(Model))
+	if(ensureMsgf(InObject, TEXT("Expected valid object, property tree will nopt be drawn")))
 	{
-		return;
+		return;	
 	}
 	
-	SlateIM::BeginScrollBox();
+	if (SlateIM::BeginTableRowChildren())
 	{
-		SlateIM::Text(TEXT("Motor Sim Components"), FColorList::Green);
-
-		for(const FMotorSimEntry& MotorSimComponent : Model->SimComponents)
+		for (TFieldIterator<FProperty> PropIt(InObject->GetClass()); PropIt; ++PropIt)
 		{
-			if(const UObject* MotorSimObj = MotorSimComponent.Sim.GetObject())
+			const FProperty* Property = *PropIt;
+
+			if(!Property)
 			{
-				SlateIM::BeginHorizontalStack();
+				continue;
+			}
+
+			const bool bIsEditableProperty = Property->HasAllPropertyFlags(CPF_Edit) && !Property->HasAnyPropertyFlags(CPF_EditConst);
+			const bool bUISupportedType = AudioMotorModelDebuggerPropertyUI::CanDrawProperty(*Property) || CastField<FArrayProperty>(Property);
+			bool bHiddenByMetadata = false;
+                                
+#if WITH_METADATA
+			bHiddenByMetadata = AudioMotorModelDebugWidgetPrivate::IsPropertyHiddenByClassCategory(*Property, *UAudioMotorSimComponent::StaticClass());								
+#endif
+								
+			if (bIsEditableProperty && bUISupportedType && !bHiddenByMetadata)
+			{
+				if(const FArrayProperty* ArrayProperty = CastField<FArrayProperty>(Property))
 				{
-					auto FindMotorSimDebugContext = [MotorSimObj](const AudioMotorModelDebug::FDebugContext& DebugContext)
-					{
-							return MotorSimObj == DebugContext.DebuggedObject;
-					};
+					bool bSupportedArrayType = AudioMotorModelDebuggerPropertyUI::CanDrawProperty(*ArrayProperty->Inner);
 
-					const AudioMotorModelDebug::FDebugContext* DebugContext = MotorSimDebugContexts.FindByPredicate(FindMotorSimDebugContext);
+					if(!bSupportedArrayType)
+					{
+						continue;
+					}
+										
+					if (SlateIM::NextTableCell())
+					{
+						SlateIM::Text(ArrayProperty->GetName());
+						if (SlateIM::BeginTableRowChildren())
+						{
+							constexpr bool bDisplayArrayName = false;
+							AudioMotorModelDebuggerPropertyUI::DrawArrayPropertyUI(*ArrayProperty, ArrayProperty->ContainerPtrToValuePtr<void>(InObject.Get()), bDisplayArrayName);
+						}
+						SlateIM::EndTableRowChildren();
+					}
+				}
+				else
+				{
+					if (SlateIM::NextTableCell())
+					{
+						AudioMotorModelDebuggerPropertyUI::DrawPropertyUI(*Property, InObject.Get());
+					}
+				}
+			}
+		}
+	}
+	SlateIM::EndTableRowChildren();
+}
 
-					const bool bIsMotorSimDrawn = DebugContext != nullptr;
-					bool bCheckBox = bIsMotorSimDrawn;
-					SlateIM::CheckBox(TEXT(""), bCheckBox);
-					FString ComponentName = FString::Printf(TEXT("%s"), *MotorSimObj->GetName());
-					SlateIM::Text(ComponentName);
-					
-					if(bIsMotorSimDrawn != bCheckBox)
+void FAudioMotorModelDebugWidget::DrawMotorSimComponentList()
+{
+	SlateIM::BeginVerticalStack();
+	{
+		SlateIM::Text(TEXT("Motor Sim Components"), FColorList::Green);
+		
+		SlateIM::BeginTable();
+		{
+			SlateIM::AddTableColumn(TEXT("Components"));
+			{
+				for(AudioMotorModelDebug::FDebugContext& DebugContext : MotorSimDebugContexts)
+				{
+					if(UAudioMotorSimComponent* DebuggedMotorSimComponent = Cast<UAudioMotorSimComponent>(DebugContext.DebuggedObject))
 					{
-						if(bIsMotorSimDrawn)
+						bool bIsMotorSimEnabled = DebuggedMotorSimComponent->bEnabled;
+
+						if (SlateIM::NextTableCell())
 						{
-							const int32 DebugContextIndex = MotorSimDebugContexts.Find(*DebugContext);
-							MotorSimDebugContexts.RemoveAtSwap(DebugContextIndex);
+							//Draws motor sim name and a checkboxes to enable/disable it and print its params
+							SlateIM::BeginHorizontalStack();
+							{
+								SlateIM::SetToolTip(TEXT("Is Enabled"));
+								SlateIM::CheckBox(TEXT(""), bIsMotorSimEnabled);
+
+								if(bIsMotorSimEnabled != DebuggedMotorSimComponent->bEnabled)
+								{
+									DebuggedMotorSimComponent->bEnabled = bIsMotorSimEnabled;
+								}
+								
+								if(!bIsMotorSimEnabled)
+								{
+									//disable widgets related to this component
+									SlateIM::BeginDisabledState();
+								}
+								
+								SlateIM::SetToolTip(TEXT("Graph local contexts"));
+								bool bDrawGraphCheckBox = DebugContext.bDrawGraphs;
+								SlateIM::CheckBox(TEXT(""), bDrawGraphCheckBox);
+
+								if(DebugContext.bDrawGraphs != bDrawGraphCheckBox)
+								{
+									if(DebugContext.bDrawGraphs)
+									{
+										DebugContext.bDrawGraphs = false;
+									}
+									else 
+									{
+										DebugContext.bDrawGraphs = true;;
+									}
+								}
+								
+								FString ComponentName = FString::Printf(TEXT("%s"), *DebuggedMotorSimComponent->GetName());
+								SlateIM::Text(ComponentName);
+
+							}
+							SlateIM::EndHorizontalStack();
 						}
-						else 
+						
+						//Expandable tree showing properties
+						DrawObjectPropertiesTree(DebugContext.DebuggedObject);
+
+						if(!bIsMotorSimEnabled)
 						{
-							CreateMotorModelObjectDebugContext(MotorSimObj);
+							SlateIM::EndDisabledState();
 						}
 					}
 				}
-				SlateIM::EndHorizontalStack();
 			}
 		}
+		SlateIM::EndTable();
 	}
-	SlateIM::EndScrollBox();
-	
+	SlateIM::EndVerticalStack();
 }
 
 void FAudioMotorModelDebugWidget::RegisterTrackedObject(const AudioMotorModelDebug::FDebugContext& InTrackedObject)
@@ -220,6 +355,7 @@ void FAudioMotorModelDebugWidget::GenerateParamSelector(const FProperty& ParamPr
 	
 	SlateIM::BeginHorizontalStack();
 	{
+		SlateIM::VAlign(VAlign_Center);
 		SlateIM::CheckBox(TEXT(""), CheckState);
 
 		if(CheckState != bIsParamDrawn)
@@ -238,7 +374,8 @@ void FAudioMotorModelDebugWidget::GenerateParamSelector(const FProperty& ParamPr
 		}
 
 		void* ValuePtr = ParamProperty.ContainerPtrToValuePtr<void>(PropertyContainer);
-		
+
+		SlateIM::VAlign(VAlign_Center);
 		SlateIM::Text(AudioMotorModelDebugWidgetPrivate::GenerateParamString(ParamProperty, ValuePtr));
 	}
 	SlateIM::EndHorizontalStack();		
@@ -278,6 +415,11 @@ void FAudioMotorModelDebugWidget::UpdateDebugContexts()
 			continue;
 		}
 
+		if(!MotorSimDebugContext.bDrawGraphs)
+		{
+			continue;
+		}
+		
 		if(const UAudioMotorModelComponent* AsMotorModelComponent = Cast<UAudioMotorModelComponent>(MotorSimDebugContext.DebuggedObject))
 		{
 			for(FInstancedStruct& DebugData : MotorSimDebugContext.DebugData)
@@ -311,48 +453,99 @@ void FAudioMotorModelDebugWidget::UpdateDebugContexts()
 	}
 }
 
-void FAudioMotorModelDebugWidget::DrawParameterGraphs()
+void FAudioMotorModelDebugWidget::PollRegisteredModels()
 {
-	SlateIM::BeginVerticalStack();
+	for(auto It = RegisteredMotorModels.CreateIterator(); It; ++It)
 	{
-		SlateIM::Text(TEXT("Parameter Graphs"), FColorList::Green);
-		SlateIM::Text(FString::Printf(TEXT("Scale: %f"), ParamGraphScale));
-		SlateIM::Slider(ParamGraphScale, 0, ParamGraphScaleMax, 1);
+		TObjectPtr<UAudioMotorModelComponent> MotorModel = *It;
 		
-		SlateIM::Text(FString::Printf(TEXT("NumFrames: %0.f"), ParamGraphNumFrames));
-		SlateIM::Slider(ParamGraphNumFrames, 1, ParamGraphNumFramesMax, 1);
+		if(!MotorModel || !MotorModel->IsValidLowLevel())
+		{
+			It.RemoveCurrent();
+		}
 	}
-	SlateIM::EndVerticalStack();
 
-	SlateIM::BeginHorizontalStack();
+	if(RegisteredMotorModels.IsEmpty())
 	{
-		for(const FProperty* PropertyToGraph : PropertiesToGraph)
+		MotorSimDebugContexts.Empty();
+	}
+}
+
+void FAudioMotorModelDebugWidget::DrawParameterGraphs()
+{
+	SlateIM::BeginVerticalStack();
+	{
+		SlateIM::BeginVerticalStack();
 		{
-			if(!PropertyToGraph)
+			SlateIM::Text(TEXT("Parameter Graphs"), FColorList::Green);
+
+			SlateIM::BeginHorizontalStack();
 			{
-				continue;
+				SlateIM::Text(TEXT("Num Graphs in Row "));
+				SlateIM::MinWidth(50.f);
+				if (SlateIM::EditableText(GraphRowSizeText))
+				{
+					LiveGraphRowSize = FMath::Clamp(FCString::Atoi(*GraphRowSizeText), 1, MaxGraphRowSize);
+				}
 			}
-
-			SlateIM::BeginVerticalStack();
+			SlateIM::EndHorizontalStack();
+		
+			SlateIM::Text(FString::Printf(TEXT("Scale: %f"), ParamGraphScale));
+			SlateIM::Slider(ParamGraphScale, 0, ParamGraphScaleMax, 1);
+		
+			SlateIM::Text(FString::Printf(TEXT("NumFrames: %0.f"), ParamGraphNumFrames));
+			SlateIM::Slider(ParamGraphNumFrames, 1, ParamGraphNumFramesMax, 1);
+		}
+		SlateIM::EndVerticalStack();
+	
+		SlateIM::BeginVerticalStack();
+		{
+			int32 NumDrawnGraphs = 0;
+			for(const FProperty* PropertyToGraph : PropertiesToGraph)
 			{
+				if(!PropertyToGraph)
+				{
+					continue;
+				}
+			
 				for(AudioMotorModelDebug::FDebugContext& DebugContext : MotorSimDebugContexts)
 				{
+					if(!DebugContext.bDrawGraphs)
+					{
+						continue;
+					}
 		
 					for(TPair<FName, TSharedPtr<AudioMotorModelDebug::IParamGraph>>& PropertyGraphPair : DebugContext.ParamGraphs)
 					{
 						if(PropertyToGraph->GetName() == PropertyGraphPair.Key && PropertyGraphPair.Value)
 						{
+							if(NumDrawnGraphs % LiveGraphRowSize == 0)
+							{
+								SlateIM::BeginHorizontalStack();
+							}
+							
 							PropertyGraphPair.Value->Draw();
 							PropertyGraphPair.Value->SetScale(ParamGraphScale);
 							PropertyGraphPair.Value->SetNumFrames(ParamGraphNumFrames);
+							NumDrawnGraphs++;
+
+							if(NumDrawnGraphs % LiveGraphRowSize == 0)
+							{
+								SlateIM::EndHorizontalStack();
+							}
 						}
 					}
 				}
 			}
-			SlateIM::EndVerticalStack();
+
+			if(NumDrawnGraphs % LiveGraphRowSize != 0)
+			{
+				SlateIM::EndHorizontalStack();
+			}
 		}
+		SlateIM::EndVerticalStack();
 	}
-	SlateIM::EndHorizontalStack();
+	SlateIM::EndVerticalStack();
 
 }
 
@@ -486,7 +679,7 @@ void FAudioMotorModelDebugWidget::DrawParameterSelection(AudioMotorModelDebug::F
 	SlateIM::EndHorizontalStack();
 }
 
-AudioMotorModelDebug::FDebugContext* FAudioMotorModelDebugWidget::CreateMotorModelObjectDebugContext(const UObject* ObjectToTrack)
+AudioMotorModelDebug::FDebugContext* FAudioMotorModelDebugWidget::CreateMotorModelObjectDebugContext(UObject* ObjectToTrack)
 {
 	ensure(ObjectToTrack);
 
@@ -500,6 +693,22 @@ AudioMotorModelDebug::FDebugContext* FAudioMotorModelDebugWidget::CreateMotorMod
 	return &MotorSimDebugContexts[Index];
 }
 
+void FAudioMotorModelDebugWidget::CreateModelSimComponentsDebugContexts(const UAudioMotorModelComponent* Model)
+{
+	if(!ensure(Model))
+	{
+		return;
+	}
+
+	for(const FMotorSimEntry& MotorSimComponent : Model->SimComponents)
+	{
+		if(UObject* MotorSimObj = MotorSimComponent.Sim.GetObject())
+		{
+			CreateMotorModelObjectDebugContext(MotorSimObj);
+		}
+	}
+}
+
 FAudioMotorModelSlateIMWindow::FAudioMotorModelSlateIMWindow(const TCHAR* Command, const TCHAR* CommandHelp)
 	: FSlateIMWindowBase(TEXT("Audio Motor Model Slate IM Window"), FVector2f(700, 900), Command, CommandHelp)
 {
@@ -514,15 +723,14 @@ void FAudioMotorModelSlateIMWindow::DrawWindow(float DeltaTime)
 FAudioMotorModelSlateIMViewportWidget::FAudioMotorModelSlateIMViewportWidget(const TCHAR* Command, const TCHAR* CommandHelp)
 	: FSlateIMWidgetWithCommandBase(Command, CommandHelp)
 {                                                                                
-	Layout.Anchors = FAnchors(0.5f, 0);                                          
-	Layout.Alignment = FVector2f(0.2f, 0);                                       
-	Layout.Size = AudioMotorModelDebugWidgetPrivate::WindowSize;                                           
+	Layout.Anchors = AudioMotorModelDebugWidgetPrivate::ViewportAnchors;
+	Layout.Alignment = AudioMotorModelDebugWidgetPrivate::ViewportAlignment;
 }
 
 void FAudioMotorModelSlateIMViewportWidget::DrawWidget(float DeltaTime)
 {
 	if (GEngine && GEngine->GameViewport)
-	{
+	{		
 		if (SlateIM::BeginViewportRoot("AudioMotorModelSlateIMViewportWidget", GEngine->GameViewport, Layout))
 		{
 			DebugWidget.Draw();
diff --git a/Engine/Plugins/Runtime/AudioMotorSim/Source/AudioMotorSimDebug/Private/AudioMotorModelSlateIMWindow.h b/Engine/Plugins/Runtime/AudioMotorSim/Source/AudioMotorSimDebug/Private/AudioMotorModelSlateIMWindow.h
--- a/Engine/Plugins/Runtime/AudioMotorSim/Source/AudioMotorSimDebug/Private/AudioMotorModelSlateIMWindow.h
+++ b/Engine/Plugins/Runtime/AudioMotorSim/Source/AudioMotorSimDebug/Private/AudioMotorModelSlateIMWindow.h
@@ -17,14 +17,15 @@ namespace AudioMotorModelDebug
 	{
 		FDebugContext() = default;
 		
-		explicit FDebugContext(const TObjectPtr<const UObject>& InDebuggedObject)
+		explicit FDebugContext(const TObjectPtr<UObject>& InDebuggedObject)
 			: DebuggedObject(InDebuggedObject)
 		{}
 
-		FDebugContext(const TObjectPtr<const UObject>& InDebuggedObject, const TArray<FInstancedStruct>& InDebugData, const TMap<FName, TSharedPtr<IParamGraph>>& InParamGraphs)
+		FDebugContext(const TObjectPtr<UObject>& InDebuggedObject, const TArray<FInstancedStruct>& InDebugData, const TMap<FName, TSharedPtr<IParamGraph>>& InParamGraphs, bool bInDrawGraphs)
 			: DebuggedObject(InDebuggedObject)
 			, DebugData(InDebugData)
 			, ParamGraphs(InParamGraphs)
+			, bDrawGraphs(bInDrawGraphs)
 		{}
 				
 		bool operator==(const FDebugContext& A) const
@@ -32,32 +33,38 @@ namespace AudioMotorModelDebug
 			return DebuggedObject == A.DebuggedObject;
 		}
 		
-		TObjectPtr<const UObject> DebuggedObject;
+		TObjectPtr<UObject> DebuggedObject;
 		TArray<FInstancedStruct> DebugData;
 		TMap<FName, TSharedPtr<IParamGraph>> ParamGraphs;
+		bool bDrawGraphs = false;
 	};
 }
 
 class FAudioMotorModelDebugWidget
 {
 public:
+	
 	void Draw();
 	void RegisterMotorModelComponent(UAudioMotorModelComponent* Component);
 
 private:
 	void UpdateDebugContexts();
-	
+
+	void PollRegisteredModels();
 	void DrawMotorModelSelection();
 	void DrawParameterSelection(AudioMotorModelDebug::FDebugContext& DebugContext);
 
-	void InstantiateMotorModelDetails(const UAudioMotorModelComponent* Model);
-	void DrawMotorSimComponentList(const UAudioMotorModelComponent* Model);
-	AudioMotorModelDebug::FDebugContext* CreateMotorModelObjectDebugContext(const UObject* ObjectToTrack);
+	void InstantiateMotorModelDetails(UAudioMotorModelComponent* Model);
+	void DrawMotorSimComponentList();
+	AudioMotorModelDebug::FDebugContext* CreateMotorModelObjectDebugContext(UObject* ObjectToTrack);
+	void CreateModelSimComponentsDebugContexts(const UAudioMotorModelComponent* Model);
 	
 	void GenerateParamSelector(const FProperty& ParamProperty, void* PropertyContainer);
 	void DrawParameterGraphs();
 	void PollParameterGraphs();
 	
+	void DrawObjectPropertiesTree(const TObjectPtr<UObject>& InObject);
+	
 	TSharedPtr<AudioMotorModelDebug::IParamGraph> CreateParamGraph(const FProperty& ParamProperty, void* PropertyContainer);
 	TSharedPtr<AudioMotorModelDebug::IParamGraph> CreateParamGraph(const FName& GraphedPropertyName, const FProperty& ParamProperty, void* PropertyContainer);
 	
@@ -68,6 +75,11 @@ private:
 	
 	int32 SelectedMotorModel = 0;
 	
+	int32 GraphRowSize = 2;
+	int32 MaxGraphRowSize = 50;
+	int32 LiveGraphRowSize = GraphRowSize;
+	FString GraphRowSizeText = FString::FromInt(GraphRowSize);
+	
 	float ParamGraphScale = 1.0f;         
 	float ParamGraphScaleMax = 5.0f;
 	
