[AudioMotorModelDebugger] Fix for property trees not displaying

#rnx

diff --git a/Engine/Plugins/Runtime/AudioMotorSim/Source/AudioMotorSimDebug/Private/AudioMotorModelSlateIMWindow.cpp b/Engine/Plugins/Runtime/AudioMotorSim/Source/AudioMotorSimDebug/Private/AudioMotorModelSlateIMWindow.cpp
--- a/Engine/Plugins/Runtime/AudioMotorSim/Source/AudioMotorSimDebug/Private/AudioMotorModelSlateIMWindow.cpp
+++ b/Engine/Plugins/Runtime/AudioMotorSim/Source/AudioMotorSimDebug/Private/AudioMotorModelSlateIMWindow.cpp
@@ -177,11 +177,8 @@ void FAudioMotorModelDebugWidget::InstantiateMotorModelDetails(UAudioMotorModelC
 	{
 		return;
 	}
-
-	if(MotorModelDebugContext)
-	{
-		MotorModelDebugContext->bDrawGraphs = true;
-	}
+	
+	MotorModelDebugContext->bDrawGraphs = true;
 	
 	SlateIM::BeginHorizontalStack();
 	{
@@ -195,7 +192,7 @@ void FAudioMotorModelDebugWidget::InstantiateMotorModelDetails(UAudioMotorModelC
 
 void FAudioMotorModelDebugWidget::DrawObjectPropertiesTree(const TObjectPtr<UObject>& InObject)
 {
-	if(ensureMsgf(InObject, TEXT("Expected valid object, property tree will nopt be drawn")))
+	if(!ensureMsgf(InObject, TEXT("Expected valid object, property tree will nopt be drawn")))
 	{
 		return;	
 	}
