Add stronger checks around the use of motor sim components

#rnx

diff --git a/Engine/Plugins/Runtime/AudioMotorSim/Source/AudioMotorSim/Private/AudioMotorModelComponent.cpp b/Engine/Plugins/Runtime/AudioMotorSim/Source/AudioMotorSim/Private/AudioMotorModelComponent.cpp
--- a/Engine/Plugins/Runtime/AudioMotorSim/Source/AudioMotorSim/Private/AudioMotorModelComponent.cpp
+++ b/Engine/Plugins/Runtime/AudioMotorSim/Source/AudioMotorSim/Private/AudioMotorModelComponent.cpp
@@ -162,7 +162,7 @@ void UAudioMotorModelComponent::Reset()
 {
 	for(FMotorSimEntry& Entry : SimComponents)
 	{
-		if(Entry.Sim)
+		if(Entry.Sim.GetObject() && Entry.Sim.GetObject()->IsValidLowLevel())
 		{
 			Entry.Sim->Reset();
 		}
@@ -228,6 +228,7 @@ void UAudioMotorModelComponent::AddMotorSimComponent(TScriptInterface<IAudioMoto
 	}
 	
 	SimComponents.Add(NewEntry);
+	NewEntry.Sim->Init();
 }
 	
 void UAudioMotorModelComponent::RemoveMotorSimComponent(TScriptInterface<IAudioMotorSim> InComponent)
@@ -248,7 +249,7 @@ void UAudioMotorModelComponent::ConfigureMotorSimComponents(const TArray<FInstan
 	{
 		for(FMotorSimEntry& SimComponent : SimComponents)
 		{
-			if(SimComponent.Sim)
+			if(SimComponent.Sim.GetObject() && SimComponent.Sim.GetObject()->IsValidLowLevel())
 			{
 				SimComponent.Sim->ConfigMotorSim(ConfigData);
 			}
