[MotorPhysicsSimComponent] Allow for collection of additional debug data
#rnx

diff --git a/Engine/Plugins/Runtime/AudioMotorSim/Source/AudioMotorSimStandardComponents/Private/MotorPhysicsSimComponent.cpp b/Engine/Plugins/Runtime/AudioMotorSim/Source/AudioMotorSimStandardComponents/Private/MotorPhysicsSimComponent.cpp
--- a/Engine/Plugins/Runtime/AudioMotorSim/Source/AudioMotorSimStandardComponents/Private/MotorPhysicsSimComponent.cpp
+++ b/Engine/Plugins/Runtime/AudioMotorSim/Source/AudioMotorSimStandardComponents/Private/MotorPhysicsSimComponent.cpp
@@ -73,6 +73,26 @@ void UMotorPhysicsSimComponent::Update(FAudioMotorSimInputContext& Input, FAudio
 
 	RuntimeInfo.Rpm = FMath::FInterpTo(RuntimeInfo.Rpm, Rpm, Input.DeltaTime, RpmInterpSpeed);
 
+#if !UE_BUILD_SHIPPING
+	if(ShouldCollectDebugData())
+	{
+		if(!CachedDebugData.IsValid())
+		{
+			CachedDebugData.InitializeAs<FAudioMotorSimDebugDataBase>();;
+		}
+		
+		if(FAudioMotorSimDebugDataBase* DebugData = CachedDebugData.GetMutablePtr<FAudioMotorSimDebugDataBase>())
+		{
+			DebugData->ParameterValues.Add("CurrentEngineTorque", CurrentEngineTorque);
+			DebugData->ParameterValues.Add("BrakingForce", BrakingForce);
+			DebugData->ParameterValues.Add("WindResistance", WindResistance);
+			DebugData->ParameterValues.Add("FrictionResistance", FrictionResistance);
+			DebugData->ParameterValues.Add("EngineResistance", EngineResistance);
+			DebugData->ParameterValues.Add("NetForce", NetForce);
+		}
+	}
+#endif
+
 	Super::Update(Input, RuntimeInfo);
 }
 
