Minor fixes for audio motor model debug window:
- uniquely identify motor sim component window
- adjust windows size

#rnx

diff --git a/Engine/Plugins/Runtime/AudioMotorSim/Source/AudioMotorSimDebug/Private/AudioMotorModelDebugWidget.cpp b/Engine/Plugins/Runtime/AudioMotorSim/Source/AudioMotorSimDebug/Private/AudioMotorModelDebugWidget.cpp
--- a/Engine/Plugins/Runtime/AudioMotorSim/Source/AudioMotorSimDebug/Private/AudioMotorModelDebugWidget.cpp
+++ b/Engine/Plugins/Runtime/AudioMotorSim/Source/AudioMotorSimDebug/Private/AudioMotorModelDebugWidget.cpp
@@ -78,7 +78,8 @@ namespace AudioMotorModelDebugWidgetPrivate
 
 
 FAudioMotorModelDebugWidget::FAudioMotorModelDebugWidget()
-{}
+{
+}
 
 void FAudioMotorModelDebugWidget::Draw()
 {
@@ -272,7 +273,8 @@ void FAudioMotorModelDebugWidget::DrawMotorSimComponentDetailButton(AudioMotorMo
 		}
 		else
 		{
-			DebugContext.DetailWindow = MakeShared<FMotorSimComponentDetailWindow>(DebuggedMotorSimComponent->GetName());
+			FString WindowName = FString::Printf(TEXT("%s (%s)"), *DebuggedMotorSimComponent->GetName(), *DebuggedMotorSimComponent->GetOuter()->GetName());
+			DebugContext.DetailWindow = MakeShared<FMotorSimComponentDetailWindow>(WindowName);
 			DebugContext.DetailWindow->ToggleWidget();
 		}
 	}
diff --git a/Engine/Plugins/Runtime/AudioMotorSim/Source/AudioMotorSimDebug/Private/AudioMotorModelDebugWidgetChildWindow.h b/Engine/Plugins/Runtime/AudioMotorSim/Source/AudioMotorSimDebug/Private/AudioMotorModelDebugWidgetChildWindow.h
--- a/Engine/Plugins/Runtime/AudioMotorSim/Source/AudioMotorSimDebug/Private/AudioMotorModelDebugWidgetChildWindow.h
+++ b/Engine/Plugins/Runtime/AudioMotorSim/Source/AudioMotorSimDebug/Private/AudioMotorModelDebugWidgetChildWindow.h
@@ -5,12 +5,13 @@
 #include "SlateIM.h"
 #include "SlateIMWidgetBase.h"
 #include "StructUtils/InstancedStruct.h"
+#include "UObject/ObjectKey.h"
 
 class FAudioMotorModelDebugWidgetChildWindow : public FSlateIMWidgetBase
 {
 public:
 	FAudioMotorModelDebugWidgetChildWindow(const FStringView& WindowTitle, FVector2f WindowSize)
-		: FSlateIMWidgetBase(FGuid().ToString())
+		: FSlateIMWidgetBase(WindowTitle)
 		, WindowTitle(WindowTitle)
 		, WindowSize(WindowSize)
 	{}
diff --git a/Engine/Plugins/Runtime/AudioMotorSim/Source/AudioMotorSimDebug/Private/AudioMotorModelSlateIMWindow.cpp b/Engine/Plugins/Runtime/AudioMotorSim/Source/AudioMotorSimDebug/Private/AudioMotorModelSlateIMWindow.cpp
--- a/Engine/Plugins/Runtime/AudioMotorSim/Source/AudioMotorSimDebug/Private/AudioMotorModelSlateIMWindow.cpp
+++ b/Engine/Plugins/Runtime/AudioMotorSim/Source/AudioMotorSimDebug/Private/AudioMotorModelSlateIMWindow.cpp
@@ -15,7 +15,7 @@ namespace AudioMotorModelSlateIMWindowPrivate
 }
 
 FAudioMotorModelSlateIMWindow::FAudioMotorModelSlateIMWindow(const TCHAR* Command, const TCHAR* CommandHelp)
-	: FSlateIMWindowBase(TEXT("Audio Motor Model Slate IM Window"), FVector2f(700, 900), Command, CommandHelp)
+	: FSlateIMWindowBase(TEXT("Audio Motor Model Slate IM Window"), AudioMotorModelSlateIMWindowPrivate::WindowSize, Command, CommandHelp)
 {}
 
 void FAudioMotorModelSlateIMWindow::DrawWindow(float DeltaTime)
