AudioMotorModel V2
- added description and metadata to sim component parameters
- send useful data to debugger
- fix wrong formulas in MotorPhysicsSimComponent
- allow downshifting to 0 in RevLimiter

#rnx

diff --git a/Engine/Plugins/Runtime/AudioMotorSim/Source/AudioMotorSim/Public/AudioMotorSimTypes.h b/Engine/Plugins/Runtime/AudioMotorSim/Source/AudioMotorSim/Public/AudioMotorSimTypes.h
--- a/Engine/Plugins/Runtime/AudioMotorSim/Source/AudioMotorSim/Public/AudioMotorSimTypes.h
+++ b/Engine/Plugins/Runtime/AudioMotorSim/Source/AudioMotorSim/Public/AudioMotorSimTypes.h
@@ -80,6 +80,10 @@ struct FAudioMotorSimRuntimeContext
 	UPROPERTY(EditAnywhere, BlueprintReadWrite, Category = "State")
 	int32 Gear = 0;
 	
+	// the virtual speed calculated by the Motor Physics Simulation
+	UPROPERTY(VisibleAnywhere, BlueprintReadWrite, Category = "State", meta = (ForceUnits = "km/h"))
+    float VirtualSpeedKmh = 0.f;
+	
 	// normalized RPM [0-1] of the motor
 	UPROPERTY(EditAnywhere, BlueprintReadWrite, Category = "State")
 	float Rpm = 0.f;
diff --git a/Engine/Plugins/Runtime/AudioMotorSim/Source/AudioMotorSimStandardComponents/Private/MotorPhysicsSimComponent.cpp b/Engine/Plugins/Runtime/AudioMotorSim/Source/AudioMotorSimStandardComponents/Private/MotorPhysicsSimComponent.cpp
--- a/Engine/Plugins/Runtime/AudioMotorSim/Source/AudioMotorSimStandardComponents/Private/MotorPhysicsSimComponent.cpp
+++ b/Engine/Plugins/Runtime/AudioMotorSim/Source/AudioMotorSimStandardComponents/Private/MotorPhysicsSimComponent.cpp
@@ -18,8 +18,6 @@ void UMotorPhysicsSimComponent::Update(FAudioMotorSimInputContext& Input, FAudio
 	const float GearRatio = GetGearRatio(RuntimeInfo.Gear, Input.bClutchEngaged);
 	const float Throttle = FMath::Abs(Input.Throttle);
 	
-	const float SpeedKmh = AudioMotorSim::CmsToKmh(Input.Speed);
-
 	int32& Gear = RuntimeInfo.Gear;
 	const int32 PrevGear = Gear;
 	
@@ -30,9 +28,9 @@ void UMotorPhysicsSimComponent::Update(FAudioMotorSimInputContext& Input, FAudio
 	const float CurrentEngineTorque = EngineTorque * Throttle * GearRatio * ClutchModifier;
 	
 	const float BrakingForce = Input.Brake * BrakingHorsePower;
-	const float WindResistance = WindResistancePerVelocity * FMath::Square(SpeedKmh);
-	const float FrictionResistance = Weight * GroundFriction * Input.SurfaceFrictionModifier * ClutchDecelModifier;
-	const float EngineResistance = CurrentEngineTorque * FMath::Max(Input.MotorFrictionModifier - Throttle, 0.0f) * EngineFriction;
+	const float WindResistance = WindResistancePerVelocity * FMath::Square(RuntimeInfo.VirtualSpeedKmh);	
+	const float FrictionResistance = Weight * GroundFriction * Input.SurfaceFrictionModifier;
+	const float EngineResistance = EngineTorque * FMath::Max(Input.MotorFrictionModifier - Throttle, 0.0f) * EngineFriction * ClutchDecelModifier;
 	
 	const float NetForce = CurrentEngineTorque - BrakingForce - WindResistance - FrictionResistance - EngineResistance;
 
@@ -42,10 +40,11 @@ void UMotorPhysicsSimComponent::Update(FAudioMotorSimInputContext& Input, FAudio
 	float VirtualSpeedKmh =  CalcVelocity(GearRatio, RuntimeInfo.Rpm);
 	VirtualSpeedKmh += ((NetForce / Weight) * Input.DeltaTime);
 	VirtualSpeedKmh = FMath::Max(VirtualSpeedKmh, 0.0f);
+	RuntimeInfo.VirtualSpeedKmh = VirtualSpeedKmh;
 
 	// Now working backwards, what is the engine RPM as a results of this velocity and gear ratio.
 	float Rpm = CalcRpm(GearRatio, VirtualSpeedKmh);
-
+	
 	UE_LOG(LogAudioMotorSim, Verbose, TEXT("[UMotorPhysicsSimComponent] Curr Gear %d - RPM: %f"), Gear, Rpm);
 	UE_LOG(LogAudioMotorSim, Verbose, TEXT("[UMotorPhysicsSimComponent] Prev Gear %d - RPM: %f"), (Gear >= 1 ? Gear - 1 : 0), CalcRpm(GetGearRatio((Gear >= 1 ? Gear - 1 : 0), Input.bClutchEngaged), VirtualSpeedKmh));
 	
@@ -89,6 +88,12 @@ void UMotorPhysicsSimComponent::Update(FAudioMotorSimInputContext& Input, FAudio
 			DebugData->ParameterValues.Add("FrictionResistance", FrictionResistance);
 			DebugData->ParameterValues.Add("EngineResistance", EngineResistance);
 			DebugData->ParameterValues.Add("NetForce", NetForce);
+			DebugData->ParameterValues.Add("VirtualSpeedKmh", VirtualSpeedKmh);
+			DebugData->ParameterValues.Add("ClutchDecelModifier", ClutchDecelModifier);
+			DebugData->ParameterValues.Add("Gear", Gear);
+			DebugData->ParameterValues.Add("GearRatio", GetGearRatio(Gear, Input.bClutchEngaged));
+			DebugData->ParameterValues.Add("RPM", GetGearRatio(Gear, Input.bClutchEngaged));
+
 		}
 	}
 #endif
diff --git a/Engine/Plugins/Runtime/AudioMotorSim/Source/AudioMotorSimStandardComponents/Private/ResistanceMotorSimComponent.cpp b/Engine/Plugins/Runtime/AudioMotorSim/Source/AudioMotorSimStandardComponents/Private/ResistanceMotorSimComponent.cpp
--- a/Engine/Plugins/Runtime/AudioMotorSim/Source/AudioMotorSimStandardComponents/Private/ResistanceMotorSimComponent.cpp
+++ b/Engine/Plugins/Runtime/AudioMotorSim/Source/AudioMotorSimStandardComponents/Private/ResistanceMotorSimComponent.cpp
@@ -21,6 +21,26 @@ void UResistanceMotorSimComponent::Update(FAudioMotorSimInputContext& Input, FAu
 		const float ZFriction = UpSpeedMaxFriction * UpSpeedRatio;
 		
 		Input.SurfaceFrictionModifier += ZFriction + SideFriction;
+		
+#if !UE_BUILD_SHIPPING
+		if(ShouldCollectDebugData())
+		{
+			if(!CachedDebugData.IsValid())
+			{
+				CachedDebugData.InitializeAs<FAudioMotorSimDebugDataBase>();;
+			}
+		
+			if(FAudioMotorSimDebugDataBase* DebugData = CachedDebugData.GetMutablePtr<FAudioMotorSimDebugDataBase>())
+			{
+				DebugData->ParameterValues.Add("SideSpeed", Input.SideSpeed);
+				DebugData->ParameterValues.Add("UpSpeed", Input.UpSpeed);
+				DebugData->ParameterValues.Add("Speed", Input.Speed);
+				DebugData->ParameterValues.Add("UpSpeed Ratio (UpSpeed / Speed)", UpSpeedRatio);
+				DebugData->ParameterValues.Add("ZFriction", ZFriction);
+				DebugData->ParameterValues.Add("SideFriction", SideFriction);
+			}
+		}
+#endif
 	}
 
 	Super::Update(Input, RuntimeInfo);
diff --git a/Engine/Plugins/Runtime/AudioMotorSim/Source/AudioMotorSimStandardComponents/Private/RevLimiterMotorSimComponent.cpp b/Engine/Plugins/Runtime/AudioMotorSim/Source/AudioMotorSimStandardComponents/Private/RevLimiterMotorSimComponent.cpp
--- a/Engine/Plugins/Runtime/AudioMotorSim/Source/AudioMotorSimStandardComponents/Private/RevLimiterMotorSimComponent.cpp
+++ b/Engine/Plugins/Runtime/AudioMotorSim/Source/AudioMotorSimStandardComponents/Private/RevLimiterMotorSimComponent.cpp
@@ -25,9 +25,17 @@ void URevLimiterMotorSimComponent::RevLimitToTarget(FAudioMotorSimInputContext&
 	Input.Throttle = bActive ? 0.0f : Input.Throttle;
 }
 
+void URevLimiterMotorSimComponent::EngageClutchAndShiftBackToZero(FAudioMotorSimInputContext& Input, FAudioMotorSimRuntimeContext& RuntimeInfo) const
+{
+	Input.bClutchEngaged = true;
+	if(bShiftBackToZero)
+	{
+		RuntimeInfo.Gear = 0;
+	}
+}
+
 void URevLimiterMotorSimComponent::Update(FAudioMotorSimInputContext& Input, FAudioMotorSimRuntimeContext& RuntimeInfo)
 {
-	
 	if(bRevLimitOnClutchEngaged && Input.bClutchEngaged)
 	{
 		RevLimitToTarget(Input, RuntimeInfo);
@@ -56,12 +64,12 @@ void URevLimiterMotorSimComponent::Update(FAudioMotorSimInputContext& Input, FAu
 	}
 	
 	Input.bCanShift = false;
-
+	
 	// We've hit the limiter
 	if (RuntimeInfo.Rpm >= LimiterMaxRpm)
 	{
 		TimeRemaining = LimitTime;
-		RuntimeInfo.Rpm = LimiterMaxRpm;
+		RuntimeInfo.Rpm = LimiterMaxRpm;				
 		OnRevLimiterHit.Broadcast();
 	}
 
@@ -70,13 +78,14 @@ void URevLimiterMotorSimComponent::Update(FAudioMotorSimInputContext& Input, FAu
 		Input.Throttle = 0.0f;
 
 		TimeRemaining -= Input.DeltaTime;
-		Input.bClutchEngaged = true;
+		
+		EngageClutchAndShiftBackToZero(Input, RuntimeInfo);
 	}
 	else if (Input.bDriving == false)
 	{
 		Input.bClutchEngaged = true;
 	}
-
+	
 	if (Input.bGrounded)
 	{
 		TimeInAir = 0.0f;
@@ -84,8 +93,8 @@ void URevLimiterMotorSimComponent::Update(FAudioMotorSimInputContext& Input, FAu
 		return;
 	}
 
-	Input.bClutchEngaged = true;
-
+	EngageClutchAndShiftBackToZero(Input, RuntimeInfo);
+	
 	if (TimeRemaining > 0.0f)
 	{
 		TimeInAir += Input.DeltaTime;
@@ -95,6 +104,24 @@ void URevLimiterMotorSimComponent::Update(FAudioMotorSimInputContext& Input, FAu
 	{
 		Input.Throttle = 0.0f;
 	}
+	
+#if !UE_BUILD_SHIPPING
+	if(ShouldCollectDebugData())
+	{
+		if(!CachedDebugData.IsValid())
+		{
+			CachedDebugData.InitializeAs<FAudioMotorSimDebugDataBase>();;
+		}
+		
+		if(FAudioMotorSimDebugDataBase* DebugData = CachedDebugData.GetMutablePtr<FAudioMotorSimDebugDataBase>())
+		{
+			DebugData->ParameterValues.Add("TimeRemaining", TimeRemaining);
+			DebugData->ParameterValues.Add("TimeInAir", TimeInAir);
+			DebugData->ParameterValues.Add("bActive", bActive);
+			DebugData->ParameterValues.Add("bRevLimitOnClutchEngaged", bRevLimitOnClutchEngaged);
+		}
+	}
+#endif
 
 	Super::Update(Input, RuntimeInfo);
 }
diff --git a/Engine/Plugins/Runtime/AudioMotorSim/Source/AudioMotorSimStandardComponents/Public/MotorPhysicsSimComponent.h b/Engine/Plugins/Runtime/AudioMotorSim/Source/AudioMotorSimStandardComponents/Public/MotorPhysicsSimComponent.h
--- a/Engine/Plugins/Runtime/AudioMotorSim/Source/AudioMotorSimStandardComponents/Public/MotorPhysicsSimComponent.h
+++ b/Engine/Plugins/Runtime/AudioMotorSim/Source/AudioMotorSimStandardComponents/Public/MotorPhysicsSimComponent.h
@@ -2,7 +2,6 @@
 #pragma once
 
 #include "IAudioMotorSim.h"
-#include "StructUtils/InstancedStruct.h"
 
 #include "MotorPhysicsSimComponent.generated.h"
 
@@ -13,66 +12,81 @@ UCLASS(ClassGroup = "AudioMotorSim", meta = (BlueprintSpawnableComponent))
 class UMotorPhysicsSimComponent : public UAudioMotorSimComponent
 {
 	GENERATED_BODY()
-public:
 
-	UPROPERTY(EditAnywhere, BlueprintReadWrite, Category = "Force")
+public:
+	//The weight of the vehicle, used to calculate the ground friction and generated RPM
+	UPROPERTY(EditAnywhere, BlueprintReadWrite, Category = "Force", meta=(ClampMin = "0.0"))
 	float Weight = 900.f;
 
-	UPROPERTY(EditAnywhere, BlueprintReadWrite, Category = "Force")
+	//The force of the engine. Large values will accelerate the vehicle more quickly, but also apply more friction when engine braking
+	UPROPERTY(EditAnywhere, BlueprintReadWrite, Category = "Force", meta=(ClampMin = "0.0"))
 	float EngineTorque = 2500.f;
 
-	UPROPERTY(EditAnywhere, BlueprintReadWrite, Category = "Force")
+	//The force applied when braking. Large values will slow down the vehicle more quickly
+	UPROPERTY(EditAnywhere, BlueprintReadWrite, Category = "Force", meta=(ClampMin = "0.0"))
 	float BrakingHorsePower = 6000.f;
 
+	//How each gear responds to speed, expressed as a ratio. Higher number means a larger RPM range within that gear and quicker progression through the gear
 	UPROPERTY(EditAnywhere, BlueprintReadWrite, Category = "Gears")
-	TArray<float> GearRatios = { 3.5f, 2.f, 1.4f, 1.f, .7f };
+	TArray<float> GearRatios = {3.5f, 2.f, 1.4f, 1.f, .7f};
 
-	UPROPERTY(EditAnywhere, BlueprintReadWrite, Category = "Gears")
+	//The gear ratio to apply when clutched
+	UPROPERTY(EditAnywhere, BlueprintReadWrite, Category = "Gears", meta=(ClampMin = "0.0"))
 	float ClutchedGearRatio = 10.f;
-	
+
+	//If true, the model will keep upshifting past the last defined gear. Gear ratio will be defined by the InfiniteGearRatio param
 	UPROPERTY(EditAnywhere, BlueprintReadWrite, Category = "Gears")
 	bool bUseInfiniteGears = false;
-	
+
+	//If enabled, will always downshift to 0th gear when a downshift occurs
 	UPROPERTY(EditAnywhere, BlueprintReadWrite, Category = "Gears")
 	bool bAlwaysDownshiftToZerothGear = false;
-	
+
 	// how much to scale gear ratio per gear past the max gear
-	UPROPERTY(EditAnywhere, BlueprintReadWrite, Category = "Gears")
+	UPROPERTY(EditAnywhere, BlueprintReadWrite, Category = "Gears", meta=(ClampMin = "0.0"))
 	float InfiniteGearRatio = 0.9f;
-	
-	UPROPERTY(EditAnywhere, BlueprintReadWrite, Category = "Gears")
+
+	//The rpm at which an upshift occurs
+	UPROPERTY(EditAnywhere, BlueprintReadWrite, Category = "Gears", meta=(ClampMin = "0.0", ClampMax=1.0))
 	float UpShiftMaxRpm = 0.97f;
 
-	UPROPERTY(EditAnywhere, BlueprintReadWrite, Category = "Gears")
+	//The rpm at which a downshift occurs
+	UPROPERTY(EditAnywhere, BlueprintReadWrite, Category = "Gears", meta=(ClampMin = "0.0", ClampMax=1.0))
 	float DownShiftStartRpm = 0.94f;
 
-	UPROPERTY(EditAnywhere, BlueprintReadWrite, Category = "Resistance")
+	// Engine Torque multiplier applied when clutched
+	UPROPERTY(EditAnywhere, BlueprintReadWrite, Category = "Resistance", meta=(ClampMin = "0.0"))
 	float ClutchedForceModifier = 1.f;
 
-	// how much more quickly the engine decelerates off-throttle while clutched
-	UPROPERTY(EditAnywhere, BlueprintReadWrite, Category = "Resistance")
+	// How much more quickly the engine decelerates off-throttle while clutched
+	UPROPERTY(EditAnywhere, BlueprintReadWrite, Category = "Resistance", meta=(ClampMin = "0.0"))
 	float ClutchedDecelScale = 1.f;
-	
-	UPROPERTY(EditAnywhere, BlueprintReadWrite, Category = "Resistance")
+
+	//The main gear ratio of the engine. Corresponds to a RPM value of 1 if the vehicle is going at this speed.
+	//Speed * GearRatio / EngineGearRatio = RPM
+	UPROPERTY(EditAnywhere, BlueprintReadWrite, Category = "Resistance", meta=(ClampMin = "0.0"))
 	float EngineGearRatio = 50.f;
 
-	// How much of the torque is loss due to engine friction
-	UPROPERTY(EditAnywhere, BlueprintReadWrite, Category = "Resistance")
-	float EngineFriction = 0.66f; 
+	// How much of the torque is lost due to engine friction
+	UPROPERTY(EditAnywhere, BlueprintReadWrite, Category = "Resistance", meta=(ClampMin = "0.0"))
+	float EngineFriction = 0.66f;
 
 	// Coefficient of Rolling Resistance
-	UPROPERTY(EditAnywhere, BlueprintReadWrite, Category = "Resistance")
+	UPROPERTY(EditAnywhere, BlueprintReadWrite, Category = "Resistance", meta=(ClampMin = "0.0"))
 	float GroundFriction = 1.f;
 
-	UPROPERTY(EditAnywhere, BlueprintReadWrite, Category = "Resistance")
+	//Multiplier to apply to the velocity value to calculate wind resistance
+	UPROPERTY(EditAnywhere, BlueprintReadWrite, Category = "Resistance", meta=(ClampMin = "0.0"))
 	float WindResistancePerVelocity = 0.015f;
 
-	UPROPERTY(EditAnywhere, BlueprintReadWrite, Category = "PhysicsSim")
+	
+	UPROPERTY(EditAnywhere, BlueprintReadWrite, Category = "PhysicsSim", meta=(ClampMin = "0.0"))
 	float ThrottleInterpolationTime = 0.050f;
 
-	UPROPERTY(EditAnywhere, BlueprintReadWrite, Category = "PhysicsSim")
+	//Interpolation value to when changing RPM 
+	UPROPERTY(EditAnywhere, BlueprintReadWrite, Category = "PhysicsSim",  meta=(ClampMin = "0.0"))
 	float RpmInterpSpeed = 0.0f;
-	
+
 	UPROPERTY(BlueprintAssignable, Category = "PhysicsSim")
 	FOnGearChanged OnGearChangedEvent;
 
diff --git a/Engine/Plugins/Runtime/AudioMotorSim/Source/AudioMotorSimStandardComponents/Public/RevLimiterMotorSimComponent.h b/Engine/Plugins/Runtime/AudioMotorSim/Source/AudioMotorSimStandardComponents/Public/RevLimiterMotorSimComponent.h
--- a/Engine/Plugins/Runtime/AudioMotorSim/Source/AudioMotorSimStandardComponents/Public/RevLimiterMotorSimComponent.h
+++ b/Engine/Plugins/Runtime/AudioMotorSim/Source/AudioMotorSimStandardComponents/Public/RevLimiterMotorSimComponent.h
@@ -13,26 +13,35 @@ class URevLimiterMotorSimComponent : public UAudioMotorSimComponent
 {
 	GENERATED_BODY()
 public:
-
-	UPROPERTY(EditAnywhere, BlueprintReadWrite, Category = "RevLimiter")
+	//How long will the throttle be cut for when the limiter kicks in  
+	UPROPERTY(EditAnywhere, BlueprintReadWrite, Category = "RevLimiter", Meta=(ClampMin="0.0", ClampMax="5.0"))
 	float LimitTime = 0.f;
-
-	UPROPERTY(EditAnywhere, BlueprintReadWrite, Category = "RevLimiter")
-	float DecelScale = 0.f;
 	
-	UPROPERTY(EditAnywhere, BlueprintReadWrite, Category = "RevLimiter")
+	UPROPERTY(EditAnywhere, BlueprintReadWrite, Category = "RevLimiter", Meta=(EditCondition=false, EditConditionHides, DeprecatedProperty, DeprecationMessage="DecelScale is deprecated."))
+	float DecelScale = 0.f;
+
+	//How long should the rev limiter be enabled for while in air 
+	UPROPERTY(EditAnywhere, BlueprintReadWrite, Category = "RevLimiter", Meta=(ClampMin="0.0", ClampMax="10.0"))
 	float AirMaxThrottleTime = 0.f;
-	
+
+	//The side speed needed to trigger the rev limiter
 	UPROPERTY(EditAnywhere, BlueprintReadWrite, Category = "RevLimiter")
 	float SideSpeedThreshold = 0.f;
 
-	UPROPERTY(EditAnywhere, BlueprintReadWrite, Category = "RevLimiter")
+	//Ceiling RPM for the Rev Limiter
+	UPROPERTY(EditAnywhere, BlueprintReadWrite, Category = "RevLimiter", Meta=(ClampMin="0.0", ClampMax="1.0"))
 	float LimiterMaxRpm = 0.f;
-	
+
+	//If the car should shift back to 0 when the rev limiter is engaged
 	UPROPERTY(EditAnywhere, BlueprintReadWrite, Category = "RevLimiter")
-	bool bRevLimitOnClutchEngaged = false;
+	bool bShiftBackToZero = false;
 
+	//If enabled, rev limiter behavior will be triggered by having the clutch engaged
 	UPROPERTY(EditAnywhere, BlueprintReadWrite, Category = "RevLimiter")
+	bool bRevLimitOnClutchEngaged = false;
+
+	//When rev limiting on clutch, the RPM at which the throttle will be re-enabled
+	UPROPERTY(EditAnywhere, BlueprintReadWrite, Category = "RevLimiter", Meta=(ClampMin="0.0", ClampMax="1.0", EditCondition="bRevLimitOnClutchEngaged", EditConditionHides))
 	float RecoverRPM = 0.f;
 
 	UPROPERTY(BlueprintAssignable)
@@ -43,7 +52,8 @@ public:
 
 private:
 	void RevLimitToTarget(FAudioMotorSimInputContext& Input, FAudioMotorSimRuntimeContext& RuntimeInfo);
-	
+	void EngageClutchAndShiftBackToZero(FAudioMotorSimInputContext& Input, FAudioMotorSimRuntimeContext& RuntimeInfo) const;
+
 	// Time remaining where the limiter is forcing throttle down
 	float TimeRemaining = 0.f;
 	float TimeInAir = 0.f;
diff --git a/Engine/Plugins/Runtime/AudioMotorSim/Source/AudioMotorSimStandardComponents/Public/ReverseMotorSimComponent.h b/Engine/Plugins/Runtime/AudioMotorSim/Source/AudioMotorSimStandardComponents/Public/ReverseMotorSimComponent.h
--- a/Engine/Plugins/Runtime/AudioMotorSim/Source/AudioMotorSimStandardComponents/Public/ReverseMotorSimComponent.h
+++ b/Engine/Plugins/Runtime/AudioMotorSim/Source/AudioMotorSimStandardComponents/Public/ReverseMotorSimComponent.h
@@ -12,7 +12,7 @@ class UReverseMotorSimComponent : public UAudioMotorSimComponent
 	
 public:
 	// How much to scale engine friction by when reversing
-	UPROPERTY(EditAnywhere, BlueprintReadWrite, Category = "Reverse")
+	UPROPERTY(EditAnywhere, BlueprintReadWrite, Category = "Reverse", Meta=(ClampMin="0.0"))
 	float ReverseEngineResistanceModifier = 1.f;
 	
 	virtual void Update(FAudioMotorSimInputContext& Input, FAudioMotorSimRuntimeContext& RuntimeInfo) override;
diff --git a/Engine/Plugins/Runtime/AudioMotorSim/Source/AudioMotorSimStandardComponents/Public/RpmCurveMotorSimComponent.h b/Engine/Plugins/Runtime/AudioMotorSim/Source/AudioMotorSimStandardComponents/Public/RpmCurveMotorSimComponent.h
--- a/Engine/Plugins/Runtime/AudioMotorSim/Source/AudioMotorSimStandardComponents/Public/RpmCurveMotorSimComponent.h
+++ b/Engine/Plugins/Runtime/AudioMotorSim/Source/AudioMotorSimStandardComponents/Public/RpmCurveMotorSimComponent.h
@@ -33,7 +33,7 @@ public:
 	UPROPERTY(EditAnywhere, Category = "Gears")
 	TArray<FMotorSimGearCurve> Gears;
 	
-	UPROPERTY(EditAnywhere, BlueprintReadWrite, Category = "Gears")
+	UPROPERTY(EditAnywhere, BlueprintReadWrite, Category = "Gears", Meta=(ClampMin="0.0"))
 	float InterpSpeed = 0.f;
 
 	UPROPERTY(BlueprintAssignable)
diff --git a/Engine/Plugins/Runtime/AudioMotorSim/Source/AudioMotorSimStandardComponents/Public/ThrottleStateMotorSimComponent.h b/Engine/Plugins/Runtime/AudioMotorSim/Source/AudioMotorSimStandardComponents/Public/ThrottleStateMotorSimComponent.h
--- a/Engine/Plugins/Runtime/AudioMotorSim/Source/AudioMotorSimStandardComponents/Public/ThrottleStateMotorSimComponent.h
+++ b/Engine/Plugins/Runtime/AudioMotorSim/Source/AudioMotorSimStandardComponents/Public/ThrottleStateMotorSimComponent.h
@@ -24,7 +24,7 @@ public:
 	UPROPERTY(BlueprintAssignable, Category = "ThrottleState")
 	FOnEngineBlowoff OnEngineBlowoff;
 
-	UPROPERTY(EditAnywhere, BlueprintReadWrite, Category = "ThrottleState")
+	UPROPERTY(EditAnywhere, BlueprintReadWrite, Category = "ThrottleState", Meta=(ClampMin="0.0"))
 	float BlowoffMinThrottleTime = 1.f;
 	
 	virtual void Update(FAudioMotorSimInputContext& Input, FAudioMotorSimRuntimeContext& RuntimeInfo) override;
diff --git a/Engine/Plugins/Runtime/AudioMotorSim/Source/AudioMotorSimStandardComponents/Public/VelocitySyncMotorSimComponent.h b/Engine/Plugins/Runtime/AudioMotorSim/Source/AudioMotorSimStandardComponents/Public/VelocitySyncMotorSimComponent.h
--- a/Engine/Plugins/Runtime/AudioMotorSim/Source/AudioMotorSimStandardComponents/Public/VelocitySyncMotorSimComponent.h
+++ b/Engine/Plugins/Runtime/AudioMotorSim/Source/AudioMotorSimStandardComponents/Public/VelocitySyncMotorSimComponent.h
@@ -12,23 +12,23 @@ class UVelocitySyncMotorSimComponent : public UAudioMotorSimComponent
 	GENERATED_BODY()
 
 public:
-	UPROPERTY(EditAnywhere, BlueprintReadWrite, Category = "VelocitySync")
+	UPROPERTY(EditAnywhere, BlueprintReadWrite, Category = "VelocitySync", Meta=(ClampMin="0.0"))
 	float NoThrottleTime = 0.1f;
 
 	// Speed below which gears will be ignored, and RPM will sync directly to velocity
-	UPROPERTY(EditAnywhere, BlueprintReadWrite, Category = "VelocitySync")
+	UPROPERTY(EditAnywhere, BlueprintReadWrite, Category = "VelocitySync", Meta=(ClampMin="0.0"))
 	float SpeedThreshold = 200.f;
 
 	UPROPERTY(EditAnywhere, BlueprintReadWrite, Category = "VelocitySync")
 	FRuntimeFloatCurve SpeedToRpmCurve;
 
-	UPROPERTY(EditAnywhere, BlueprintReadWrite, Category = "VelocitySync")
+	UPROPERTY(EditAnywhere, BlueprintReadWrite, Category = "VelocitySync", Meta=(ClampMin="0.0"))
 	float InterpSpeed = 3.f;
 
-	UPROPERTY(EditAnywhere, BlueprintReadWrite, Category = "VelocitySync")
+	UPROPERTY(EditAnywhere, BlueprintReadWrite, Category = "VelocitySync", Meta=(ClampMin="0.0"))
 	float InterpTime = 0.1f;
 
-	UPROPERTY(EditAnywhere, BlueprintReadWrite, Category = "VelocitySync")
+	UPROPERTY(EditAnywhere, BlueprintReadWrite, Category = "VelocitySync", Meta=(ClampMin="0.0"))
 	float FirstGearThrottleThreshold = 0.5f;
 	
 	virtual void Update(FAudioMotorSimInputContext& Input, FAudioMotorSimRuntimeContext& RuntimeInfo) override;
