Allow vehicles to opt in into shifting in reverse and decide what RPM they  want to reach
#rnx

#changelist validated

diff --git a/Engine/Plugins/Runtime/AudioMotorSim/Source/AudioMotorSimStandardComponents/Private/RevLimiterMotorSimComponent.cpp b/Engine/Plugins/Runtime/AudioMotorSim/Source/AudioMotorSimStandardComponents/Private/RevLimiterMotorSimComponent.cpp
--- a/Engine/Plugins/Runtime/AudioMotorSim/Source/AudioMotorSimStandardComponents/Private/RevLimiterMotorSimComponent.cpp
+++ b/Engine/Plugins/Runtime/AudioMotorSim/Source/AudioMotorSimStandardComponents/Private/RevLimiterMotorSimComponent.cpp
@@ -5,9 +5,9 @@
 
 #include UE_INLINE_GENERATED_CPP_BY_NAME(RevLimiterMotorSimComponent)
 
-void URevLimiterMotorSimComponent::RevLimitToTarget(FAudioMotorSimInputContext& Input, FAudioMotorSimRuntimeContext& RuntimeInfo)
+void URevLimiterMotorSimComponent::RevLimitToTarget(FAudioMotorSimInputContext& Input, FAudioMotorSimRuntimeContext& RuntimeInfo, const float LimitRPM, const float TargetRPM)
 {
-	if (RuntimeInfo.Rpm >= LimiterMaxRpm && !bActive)
+	if (RuntimeInfo.Rpm >= LimitRPM && !bActive)
 	{
 		bActive = true;
 		TimeRemaining = LimitTime;
@@ -15,7 +15,7 @@ void URevLimiterMotorSimComponent::RevLimitToTarget(FAudioMotorSimInputContext&
 		OnRevLimiterStateChanged.Broadcast(bActive);
 	}
 
-	if(RuntimeInfo.Rpm <= RecoverRPM && bActive)
+	if(RuntimeInfo.Rpm <= TargetRPM && bActive)
 	{
 		bActive = false;
 		OnRevLimiterStateChanged.Broadcast(bActive);
@@ -38,11 +38,38 @@ void URevLimiterMotorSimComponent::Update(FAudioMotorSimInputContext& Input, FAu
 {
 	if(bRevLimitOnClutchEngaged && Input.bClutchEngaged)
 	{
-		RevLimitToTarget(Input, RuntimeInfo);
+		RevLimitToTarget(Input, RuntimeInfo, LimiterMaxRpm, ClutchedRecoverRPM);
+		Super::Update(Input, RuntimeInfo);
 		return;
 	}
+
+	const bool bIsGroundedDriving = Input.bDriving && Input.bGrounded;
+	
+	if(bRevLimitOnReverse && bIsGroundedDriving)
+	{
+		const bool bIsReversing = Input.ForwardSpeed < 0.f && Input.Throttle <= 0.f;
+
+		if(bIsReversing)
+		{
+			if(bHoldRPMInReverse)
+			{
+				//Force RPM only if throttle is actually being held down
+				if(Input.Throttle < 0)
+				{
+					RuntimeInfo.Rpm = ReverseMaxRPM;
+				}
+			}
+			else
+			{
+				RevLimitToTarget(Input, RuntimeInfo, ReverseMaxRPM, ReverseRecoverRPM);
+			}
+
+			Super::Update(Input, RuntimeInfo);
+			return;
+		}
+	}
 	
-	if (Input.bDriving && Input.bGrounded && Input.SideSpeed < SideSpeedThreshold)
+	if (bIsGroundedDriving && Input.SideSpeed < SideSpeedThreshold)
 	{
 		TimeRemaining = 0.f;
 		TimeInAir = 0.0f;
diff --git a/Engine/Plugins/Runtime/AudioMotorSim/Source/AudioMotorSimStandardComponents/Public/RevLimiterMotorSimComponent.h b/Engine/Plugins/Runtime/AudioMotorSim/Source/AudioMotorSimStandardComponents/Public/RevLimiterMotorSimComponent.h
--- a/Engine/Plugins/Runtime/AudioMotorSim/Source/AudioMotorSimStandardComponents/Public/RevLimiterMotorSimComponent.h
+++ b/Engine/Plugins/Runtime/AudioMotorSim/Source/AudioMotorSimStandardComponents/Public/RevLimiterMotorSimComponent.h
@@ -2,9 +2,11 @@
 #pragma once
 
 #include "IAudioMotorSim.h"
+
 #include "RevLimiterMotorSimComponent.generated.h"
 
 DECLARE_DYNAMIC_MULTICAST_DELEGATE(FOnRevLimiterHit);
+
 DECLARE_DYNAMIC_MULTICAST_DELEGATE_OneParam(FOnRevLimiterStateChanged, bool, bNewState);
 
 // Temporarily cuts throttle and reduces RPM when drifting or in the air
@@ -37,12 +39,28 @@ public:
 	bool bShiftBackToZero = false;
 
 	//If enabled, rev limiter behavior will be triggered by having the clutch engaged
-	UPROPERTY(EditAnywhere, BlueprintReadWrite, Category = "RevLimiter")
+	UPROPERTY(EditAnywhere, BlueprintReadWrite, Category = "RevLimiter|Clutch")
 	bool bRevLimitOnClutchEngaged = false;
 
 	//When rev limiting on clutch, the RPM at which the throttle will be re-enabled
-	UPROPERTY(EditAnywhere, BlueprintReadWrite, Category = "RevLimiter", Meta=(ClampMin="0.0", ClampMax="1.0", EditCondition="bRevLimitOnClutchEngaged", EditConditionHides))
-	float RecoverRPM = 0.f;
+	UPROPERTY(EditAnywhere, BlueprintReadWrite, Category = "RevLimiter|Clutch", Meta=(ClampMin="0.0", ClampMax="1.0", EditCondition="bRevLimitOnClutchEngaged", EditConditionHides))
+	float ClutchedRecoverRPM = 0.f;
+
+	//If enabled, rev limiter behavior will be triggered when in reverse
+	UPROPERTY(EditAnywhere, BlueprintReadWrite, Category = "RevLimiter|Reverse")
+	bool bRevLimitOnReverse = false;
+
+	//The max RPM allowed when reversing
+	UPROPERTY(EditAnywhere, BlueprintReadWrite, Category = "RevLimiter|Reverse", Meta=(ClampMin="0.0", ClampMax="1.0", EditCondition="bRevLimitOnReverse", EditConditionHides))
+	float ReverseMaxRPM = 0.f;
+	
+	//When rev limiting on reverse, the RPM at which the throttle will be re-enabled
+	UPROPERTY(EditAnywhere, BlueprintReadWrite, Category = "RevLimiter|Reverse", Meta=(ClampMin="0.0", ClampMax="1.0", EditCondition="bRevLimitOnReverse && !bHoldRPMInReverse", EditConditionHides))
+	float ReverseRecoverRPM = 0.f;
+
+	//Bypasses the limiter in reverse and holds the RPM at ReverseMaxRPM
+	UPROPERTY(EditAnywhere, BlueprintReadWrite, Category = "RevLimiter|Reverse", Meta=(EditCondition="bRevLimitOnReverse", EditConditionHides))
+	bool bHoldRPMInReverse = false;
 
 	UPROPERTY(BlueprintAssignable)
 	FOnRevLimiterHit OnRevLimiterHit;
@@ -51,7 +69,7 @@ public:
 	FOnRevLimiterStateChanged OnRevLimiterStateChanged;
 
 private:
-	void RevLimitToTarget(FAudioMotorSimInputContext& Input, FAudioMotorSimRuntimeContext& RuntimeInfo);
+	void RevLimitToTarget(FAudioMotorSimInputContext& Input, FAudioMotorSimRuntimeContext& RuntimeInfo, const float LimitRPM, const float TargetRPM);
 	void EngageClutchAndShiftBackToZero(FAudioMotorSimInputContext& Input, FAudioMotorSimRuntimeContext& RuntimeInfo) const;
 
 	// Time remaining where the limiter is forcing throttle down
