[VehicleAudio] Rename comments and some variable in MotorPhysicsSim component for clarity

#rnx

diff --git a/Engine/Plugins/Runtime/AudioMotorSim/Source/AudioMotorSimStandardComponents/Private/MotorPhysicsSimComponent.cpp b/Engine/Plugins/Runtime/AudioMotorSim/Source/AudioMotorSimStandardComponents/Private/MotorPhysicsSimComponent.cpp
--- a/Engine/Plugins/Runtime/AudioMotorSim/Source/AudioMotorSimStandardComponents/Private/MotorPhysicsSimComponent.cpp
+++ b/Engine/Plugins/Runtime/AudioMotorSim/Source/AudioMotorSimStandardComponents/Private/MotorPhysicsSimComponent.cpp
@@ -29,12 +29,12 @@ void UMotorPhysicsSimComponent::Update(FAudioMotorSimInputContext& Input, FAudio
 	
 	const float BrakingForce = Input.Brake * BrakingHorsePower;
 	const float WindResistance = WindResistancePerVelocity * FMath::Square(RuntimeInfo.VirtualSpeedKmh);	
-	const float FrictionResistance = Weight * GroundFriction * Input.SurfaceFrictionModifier;
+	const float GroundFrictionResistance = Weight * GroundFriction * Input.SurfaceFrictionModifier;
 	const float EngineResistance = EngineTorque * FMath::Max(Input.MotorFrictionModifier - Throttle, 0.0f) * EngineFriction * ClutchDecelModifier;
 	
-	const float NetForce = CurrentEngineTorque - BrakingForce - WindResistance - FrictionResistance - EngineResistance;
+	const float NetForce = CurrentEngineTorque - BrakingForce - WindResistance - GroundFrictionResistance - EngineResistance;
 
-	UE_LOG(LogAudioMotorSim, Verbose, TEXT("[UMotorPhysicsSimComponent] NetForce: %f (CurrentEngineTorque: %f - BrakingForce: %f - WindResistance %f - FrictionResistance %f - EngineResistance %f"), NetForce, CurrentEngineTorque, BrakingForce, WindResistance, FrictionResistance, EngineResistance)
+	UE_LOG(LogAudioMotorSim, Verbose, TEXT("[UMotorPhysicsSimComponent] NetForce: %f (CurrentEngineTorque: %f - BrakingForce: %f - WindResistance %f - GroundFrictionResistance %f - EngineResistance %f"), NetForce, CurrentEngineTorque, BrakingForce, WindResistance, GroundFrictionResistance, EngineResistance)
 
 	// Convert Rpm to velocity to apply forces
 	float VirtualSpeedKmh =  CalcVelocity(GearRatio, RuntimeInfo.Rpm);
@@ -85,7 +85,7 @@ void UMotorPhysicsSimComponent::Update(FAudioMotorSimInputContext& Input, FAudio
 			DebugData->ParameterValues.Add("CurrentEngineTorque", CurrentEngineTorque);
 			DebugData->ParameterValues.Add("BrakingForce", BrakingForce);
 			DebugData->ParameterValues.Add("WindResistance", WindResistance);
-			DebugData->ParameterValues.Add("FrictionResistance", FrictionResistance);
+			DebugData->ParameterValues.Add("GroundFrictionResistance", GroundFrictionResistance);
 			DebugData->ParameterValues.Add("EngineResistance", EngineResistance);
 			DebugData->ParameterValues.Add("NetForce", NetForce);
 			DebugData->ParameterValues.Add("VirtualSpeedKmh", VirtualSpeedKmh);
diff --git a/Engine/Plugins/Runtime/AudioMotorSim/Source/AudioMotorSimStandardComponents/Public/MotorPhysicsSimComponent.h b/Engine/Plugins/Runtime/AudioMotorSim/Source/AudioMotorSimStandardComponents/Public/MotorPhysicsSimComponent.h
--- a/Engine/Plugins/Runtime/AudioMotorSim/Source/AudioMotorSimStandardComponents/Public/MotorPhysicsSimComponent.h
+++ b/Engine/Plugins/Runtime/AudioMotorSim/Source/AudioMotorSimStandardComponents/Public/MotorPhysicsSimComponent.h
@@ -14,7 +14,7 @@ class UMotorPhysicsSimComponent : public UAudioMotorSimComponent
 	GENERATED_BODY()
 
 public:
-	//The weight of the vehicle, used to calculate the ground friction and generated RPM
+	//The weight of the vehicle, used to calculate the generated RPM. Heavier weight means a slower rpm increase due to weight resistance and ground friction
 	UPROPERTY(EditAnywhere, BlueprintReadWrite, Category = "Force", meta=(ClampMin = "0.0"))
 	float Weight = 900.f;
 
diff --git a/Engine/Plugins/Runtime/AudioMotorSim/Source/AudioMotorSimStandardComponents/Public/StandardAudioMotorSimConfigDataTypes.h b/Engine/Plugins/Runtime/AudioMotorSim/Source/AudioMotorSimStandardComponents/Public/StandardAudioMotorSimConfigDataTypes.h
--- a/Engine/Plugins/Runtime/AudioMotorSim/Source/AudioMotorSimStandardComponents/Public/StandardAudioMotorSimConfigDataTypes.h
+++ b/Engine/Plugins/Runtime/AudioMotorSim/Source/AudioMotorSimStandardComponents/Public/StandardAudioMotorSimConfigDataTypes.h
@@ -12,12 +12,15 @@ struct FMotorPhysicsSimConfigData : public FAudioMotorSimConfigData
 {
 	GENERATED_BODY()
 
+	// The weight of the vehicle, used to calculate the generated RPM. Heavier weight means a slower rpm increase due to weight resistance and ground friction
 	UPROPERTY(EditAnywhere, BlueprintReadOnly, Category = "Force")
 	float Weight = 900.f;
 
+	// The force of the engine. Large values will accelerate the vehicle more quickly, but also apply more friction when engine braking
 	UPROPERTY(EditAnywhere, BlueprintReadOnly, Category = "Force")
 	float EngineTorque = 2500.f;
 
+	// The force applied when braking. Large values will slow down the vehicle more quickly
 	UPROPERTY(EditAnywhere, BlueprintReadOnly, Category = "Force")
 	float BrakingHorsePower = 6000.f;
 };
