make it so coronado vehicle does not shift when locked into place

diff --git a/Engine/Plugins/Runtime/AudioMotorSim/Source/AudioMotorSimStandardComponents/Private/RevLimiterMotorSimComponent.cpp b/Engine/Plugins/Runtime/AudioMotorSim/Source/AudioMotorSimStandardComponents/Private/RevLimiterMotorSimComponent.cpp
--- a/Engine/Plugins/Runtime/AudioMotorSim/Source/AudioMotorSimStandardComponents/Private/RevLimiterMotorSimComponent.cpp
+++ b/Engine/Plugins/Runtime/AudioMotorSim/Source/AudioMotorSimStandardComponents/Private/RevLimiterMotorSimComponent.cpp
@@ -38,6 +38,7 @@ void URevLimiterMotorSimComponent::Update(FAudioMotorSimInputContext& Input, FAu
 {
 	if(bRevLimitOnClutchEngaged && Input.bClutchEngaged)
 	{
+		RuntimeInfo.Gear = bShiftBackToZero ? 0 : RuntimeInfo.Gear;
 		RevLimitToTarget(Input, RuntimeInfo, LimiterMaxRpm, ClutchedRecoverRPM);
 		Super::Update(Input, RuntimeInfo);
 		return;
