enable motor model debug on client only

diff --git a/Engine/Plugins/Runtime/AudioMotorSim/AudioMotorSim.uplugin b/Engine/Plugins/Runtime/AudioMotorSim/AudioMotorSim.uplugin
--- a/Engine/Plugins/Runtime/AudioMotorSim/AudioMotorSim.uplugin
+++ b/Engine/Plugins/Runtime/AudioMotorSim/AudioMotorSim.uplugin
@@ -29,7 +29,8 @@
 			"Name": "AudioMotorSimDebug",
 			"Type": "Runtime",
 			"LoadingPhase": "Default",
-			"TargetConfigurationDenyList": [ "Shipping" ]
+			"TargetConfigurationDenyList": [ "Shipping" ],
+			"TargetDenyList": [ "Server" ]
 		}
 	],
 	"Plugins": [
