Fix for motor sim component reverse hold ignoring physics sim

#rnx
#tests PIE

diff --git a/Engine/Plugins/Runtime/AudioMotorSim/Source/AudioMotorSimStandardComponents/Private/RevLimiterMotorSimComponent.cpp b/Engine/Plugins/Runtime/AudioMotorSim/Source/AudioMotorSimStandardComponents/Private/RevLimiterMotorSimComponent.cpp
--- a/Engine/Plugins/Runtime/AudioMotorSim/Source/AudioMotorSimStandardComponents/Private/RevLimiterMotorSimComponent.cpp
+++ b/Engine/Plugins/Runtime/AudioMotorSim/Source/AudioMotorSimStandardComponents/Private/RevLimiterMotorSimComponent.cpp
@@ -55,7 +55,7 @@ void URevLimiterMotorSimComponent::Update(FAudioMotorSimInputContext& Input, FAu
 			if(bHoldRPMInReverse)
 			{
 				//Force RPM only if throttle is actually being held down
-				if(Input.Throttle < 0)
+				if(Input.Throttle < 0 && RuntimeInfo.Rpm > ReverseMaxRPM)
 				{
 					RuntimeInfo.Rpm = ReverseMaxRPM;
 				}
