make audio motor model debugger dockable
#Synced-CL 49329988

diff --git a/Engine/Plugins/Runtime/AudioMotorSim/Source/AudioMotorSimDebug/Private/AudioMotorModelSlateIMWindow.cpp b/Engine/Plugins/Runtime/AudioMotorSim/Source/AudioMotorSimDebug/Private/AudioMotorModelSlateIMWindow.cpp
--- a/Engine/Plugins/Runtime/AudioMotorSim/Source/AudioMotorSimDebug/Private/AudioMotorModelSlateIMWindow.cpp
+++ b/Engine/Plugins/Runtime/AudioMotorSim/Source/AudioMotorSimDebug/Private/AudioMotorModelSlateIMWindow.cpp
@@ -4,6 +4,7 @@
 
 #include "Engine/Engine.h"
 #include "Engine/Texture2D.h"
+#include "Textures/SlateIcon.h"
 #include "SlateIM.h"
 
 
@@ -15,10 +16,10 @@ namespace AudioMotorModelSlateIMWindowPrivate
 }
 
 FAudioMotorModelSlateIMWindow::FAudioMotorModelSlateIMWindow(const TCHAR* Command, const TCHAR* CommandHelp)
-	: FSlateIMWindowBase(TEXT("Audio Motor Model Slate IM Window"), AudioMotorModelSlateIMWindowPrivate::WindowSize, Command, CommandHelp)
+	: FSlateIMNomadTabBase(TEXT("Audio Motor Model Slate IM Window"), Command, CommandHelp, FSlateIcon("CoreStyle", "Launcher.Tabs.Tools"))
 {}
 
-void FAudioMotorModelSlateIMWindow::DrawWindow(float DeltaTime)
+void FAudioMotorModelSlateIMWindow::DrawContent(float DeltaTime)
 {
 	DebugWidget.Draw();
 }
diff --git a/Engine/Plugins/Runtime/AudioMotorSim/Source/AudioMotorSimDebug/Private/AudioMotorModelSlateIMWindow.h b/Engine/Plugins/Runtime/AudioMotorSim/Source/AudioMotorSimDebug/Private/AudioMotorModelSlateIMWindow.h
--- a/Engine/Plugins/Runtime/AudioMotorSim/Source/AudioMotorSimDebug/Private/AudioMotorModelSlateIMWindow.h
+++ b/Engine/Plugins/Runtime/AudioMotorSim/Source/AudioMotorSimDebug/Private/AudioMotorModelSlateIMWindow.h
@@ -6,16 +6,17 @@
 #include "MotorSimComponentDetailWindow.h"
 #include "SlateIMWidgetBase.h"
 
-class FAudioMotorModelSlateIMWindow : public FSlateIMWindowBase
+class FAudioMotorModelSlateIMWindow : public FSlateIMNomadTabBase
 {
 public:
 	FAudioMotorModelSlateIMWindow(const TCHAR* Command, const TCHAR* CommandHelp);
 
-	virtual void DrawWindow(float DeltaTime) override;
+	virtual void DrawContent(float DeltaTime) override;
 
 	FAudioMotorModelDebugWidget DebugWidget;
 };
 
+
 class FAudioMotorModelSlateIMViewportWidget : public FSlateIMWidgetWithCommandBase              
 {                                                                                    
 public:                                                                              
